Fix for UE-41208
#tests Preflighted nodes successfully

diff --git a/Engine/Plugins/Performance/PerformanceMonitor/Source/PerformanceMonitor/Private/PerformanceMonitor.cpp b/Engine/Plugins/Performance/PerformanceMonitor/Source/PerformanceMonitor/Private/PerformanceMonitor.cpp
--- a/Engine/Plugins/Performance/PerformanceMonitor/Source/PerformanceMonitor/Private/PerformanceMonitor.cpp
+++ b/Engine/Plugins/Performance/PerformanceMonitor/Source/PerformanceMonitor/Private/PerformanceMonitor.cpp
@@ -3,13 +3,15 @@
 //#include "OrionGame.h"
 #include "PerformanceMonitor.h"
 #include "CoreGlobals.h"
-#include "EngineGlobals.h"
-#include "Engine/Engine.h"
 #include "HAL/FileManager.h"
 #include "Misc/Paths.h"
-#include "Engine/GameViewportClient.h"
 #include "Containers/Ticker.h"
 
+#include "Engine/Engine.h"
+#include "EngineGlobals.h"
+#include "Engine/GameViewportClient.h"
+
+
 #define SUPER_DETAILED_AUTOMATION_STATS 1
 PERFORMANCEMONITOR_API int ExportedInt;
 IMPLEMENT_MODULE(FPerformanceMonitorModule, PerformanceMonitor);
@@ -138,9 +140,9 @@ bool FPerformanceMonitorModule::Exec(UWorld* InWorld, const TCHAR* Cmd, FOutputD
 
 
 
+#if STATS
 void FPerformanceMonitorModule::StartRecordingPerfTimers(FString FileNameToUse, TArray<FString> StatsToRecord)
 {
-#if STATS
 	for (const FWorldContext& WorldContext : GEngine->GetWorldContexts())
 	{
 		if (WorldContext.WorldType == EWorldType::Game || WorldContext.WorldType == EWorldType::PIE)
@@ -213,9 +215,9 @@ void FPerformanceMonitorModule::StartRecordingPerfTimers(FString FileNameToUse,
 	//		PlayerWorld->SetTimeUntilNextGarbageCollection(300.f);
 	//	}
 	//}
-#endif
 }
 
+#endif
 void FPerformanceMonitorModule::RecordFrame()
 {
 	if (!bRecording)
diff --git a/Engine/Plugins/Performance/PerformanceMonitor/Source/PerformanceMonitor/Private/PerformanceMonitor.h b/Engine/Plugins/Performance/PerformanceMonitor/Source/PerformanceMonitor/Private/PerformanceMonitor.h
--- a/Engine/Plugins/Performance/PerformanceMonitor/Source/PerformanceMonitor/Private/PerformanceMonitor.h
+++ b/Engine/Plugins/Performance/PerformanceMonitor/Source/PerformanceMonitor/Private/PerformanceMonitor.h
@@ -22,7 +22,9 @@ class FPerformanceMonitorModule : public IModuleInterface, public FSelfRegisteri
 	};
 private: 
 	TArray<FString> DesiredStats;
+#if STATS
 	TArray<TArray<FStatMessage>> StoredMessages;
+#endif
 	TMap<FString, TArray<float>> GeneratedStats;
 	FArchive * FileToLogTo;
 	FString LogFileName;
@@ -76,7 +78,9 @@ public:
 
 
 	bool IsRecordingPerfTimers() { return bRecording; }
+#if STATS
 	void StartRecordingPerfTimers(FString FileName, TArray<FString> StatsToRecord);
+#endif
 	void SetRecordInterval(float NewInterval);
 
 	void RecordFrame();
