Copying //UE4/Dev-Rendering to //UE4/Dev-Main (Source: //UE4/Dev-Rendering @ 3284600)

==========================
MAJOR FEATURES + CHANGES
==========================

Change 3276199 on 2017/01/27 by [REDACTED_PERSON]

	DR - Typo fixes

Change 3277152 on 2017/01/30 by [REDACTED_PERSON]

	Fixed typos in various engine material functions (contributed by [REDACTED_PERSON] and [REDACTED_PERSON]).

Change 3277700 on 2017/01/30 by [REDACTED_PERSON]

	When advancing source particle index, make sure it's not already in use by another trail

Change 3278069 on 2017/01/30 by [REDACTED_PERSON]

	[Copy] Added LODDistanceFactor to SceneCaptureComponent
	* Can be used to make scene captures and planar reflections much cheaper to render if the scene has good LODs setup

Change 3278094 on 2017/01/30 by [REDACTED_PERSON]

	[Copy] Scene color resolves are restricted to the views.  Fixes inefficiency with adaptive pixel density where the views don't match the view family size.

Change 3278140 on 2017/01/30 by [REDACTED_PERSON]

	[Copy] When using a full prepass, the base pass no longer writes depth.
	* Saves .3ms of SceneDepthResolve on 970 GTX with MSAA.  Also allows masked materials to get EarlyZ in the base pass.
	* Fixed primitives with bUseAsOccluder=false being excluded from a full depth pass in static draw lists
	* Added 'r.BasePassWriteDepthEvenWithFullPrepass' for verifying that the base pass depths match the prepass depths

Change 3278166 on 2017/01/30 by [REDACTED_PERSON]

	[Copy] Quad overdraw viewmode works in the forward renderer

Change 3279100 on 2017/01/31 by [REDACTED_PERSON]

	CVar null checks.

Change 3279409 on 2017/01/31 by [REDACTED_PERSON]

	DR - vk - Fix for validation

Change 3279624 on 2017/01/31 by [REDACTED_PERSON]

	DR - Move defines out of cs into h

Change 3279895 on 2017/01/31 by [REDACTED_PERSON]

	Modified PerformanceMonitor to take all of its config from defaultgame.ini
	#tests Ran several tests successfully.

Change 3279906 on 2017/01/31 by [REDACTED_PERSON]

	Had a weird spurious declare in the config parsing function. Fixed it.
	#tests compiled.

Change 3280451 on 2017/01/31 by [REDACTED_PERSON]

	Fix Dx12 in content examples where M_Light_Bulb_Glass taking a lot of resources (9 CBVs)

Change 3281592 on 2017/02/01 by [REDACTED_PERSON]

	Update Metal SRV support to include structured-buffers which are now used by DistanceField rendering systems.

Change 3281850 on 2017/02/01 by [REDACTED_PERSON]

	Merging //UE4/Dev-Main to Dev-Rendering@3281582 (//UE4/Dev-Rendering)

Change 3282353 on 2017/02/01 by [REDACTED_PERSON]

	built lighting for QA-LightsStatic

Change 3282549 on 2017/02/01 by [REDACTED_PERSON]

	Fix DX12 crash on SkinMaterial

Change 3282637 on 2017/02/02 by [REDACTED_PERSON]

	Probably fix non-unity build

Change 3283125 on 2017/02/02 by [REDACTED_PERSON]

	Fix linux case sensitivity.

diff --git a/Engine/Plugins/Performance/PerformanceMonitor/Source/PerformanceMonitor/Private/PerformanceMonitor.cpp b/Engine/Plugins/Performance/PerformanceMonitor/Source/PerformanceMonitor/Private/PerformanceMonitor.cpp
--- a/Engine/Plugins/Performance/PerformanceMonitor/Source/PerformanceMonitor/Private/PerformanceMonitor.cpp
+++ b/Engine/Plugins/Performance/PerformanceMonitor/Source/PerformanceMonitor/Private/PerformanceMonitor.cpp
@@ -10,7 +10,7 @@
 #include "Engine/Engine.h"
 #include "EngineGlobals.h"
 #include "Engine/GameViewportClient.h"
-
+#include "ConfigCacheIni.h"
 
 #define SUPER_DETAILED_AUTOMATION_STATS 1
 PERFORMANCEMONITOR_API int ExportedInt;
@@ -19,6 +19,7 @@ FPerformanceMonitorModule::FPerformanceMonitorModule()
 {
 	bRecording = false;
 	FileToLogTo = nullptr;
+
 }
 
 FPerformanceMonitorModule::~FPerformanceMonitorModule()
@@ -140,9 +141,10 @@ bool FPerformanceMonitorModule::Exec(UWorld* InWorld, const TCHAR* Cmd, FOutputD
 
 
 
-#if STATS
 void FPerformanceMonitorModule::StartRecordingPerfTimers(FString FileNameToUse, TArray<FString> StatsToRecord)
 {
+#if STATS
+
 	for (const FWorldContext& WorldContext : GEngine->GetWorldContexts())
 	{
 		if (WorldContext.WorldType == EWorldType::Game || WorldContext.WorldType == EWorldType::PIE)
@@ -186,7 +188,55 @@ void FPerformanceMonitorModule::StartRecordingPerfTimers(FString FileNameToUse,
 		FString FileName = FString::Printf(TEXT("%s.csv"), *FileNameToUse);
 		FileToLogTo =  IFileManager::Get().CreateFileWriter(*FileName, false);
 	}
-	if (!StatsToRecord.Num())
+	if (GConfig)
+	{
+		FString ConfigCategory = FString::Printf(TEXT("/Plugins/PerformanceMonitor/%s"), *FileNameToUse);
+		float floatValueReceived = 0;
+		if (GConfig->GetFloat(*ConfigCategory,
+			TEXT("PerformanceMonitorInterval"),
+			floatValueReceived,
+			GGameIni
+			))
+		{
+			TimeBetweenRecords = floatValueReceived;
+		}
+		floatValueReceived = 0;
+		if (GConfig->GetFloat(*ConfigCategory,
+			TEXT("PerformanceMonitorTimeout"),
+			floatValueReceived,
+			GGameIni
+			))
+		{
+			TestTimeOut = floatValueReceived;
+		}
+		TArray<FString> TimersOfInterest;
+		if (GConfig->GetArray(*ConfigCategory,
+			TEXT("PerformanceMonitorTimers"), TimersOfInterest, GGameIni))
+		{
+			DesiredStats = TimersOfInterest;
+		}
+		FString MapToLoad;
+		if (GConfig->GetString(*ConfigCategory,
+			TEXT("PerformanceMonitorMap"), MapToLoad, GGameIni))
+		{
+			MapToTest = MapToLoad;
+		}
+		bExitOnCompletion = false;
+		bool GatheredBool;
+		if (GConfig->GetBool(*ConfigCategory,
+			TEXT("PerformanceMonitorExitOnFinish"), GatheredBool, GGameIni))
+		{
+			bExitOnCompletion = GatheredBool;
+		}
+		bRequiresCutsceneStart = false;
+		GatheredBool = false;
+		if (GConfig->GetBool(*ConfigCategory,
+			TEXT("PerformanceMonitorRequireCutsceneStart"), GatheredBool, GGameIni))
+		{
+			bRequiresCutsceneStart = GatheredBool;
+		}
+	}
+	if (!StatsToRecord.Num() && !DesiredStats.Num())
 	{
 		DesiredStats.Add(TEXT("STAT_AnimGameThreadTime"));
 		DesiredStats.Add(TEXT("STAT_PerformAnimEvaluation_WorkerThread"));
@@ -194,10 +244,6 @@ void FPerformanceMonitorModule::StartRecordingPerfTimers(FString FileNameToUse,
 		DesiredStats.Add(TEXT("STAT_ParticleRenderingTime"));
 		DesiredStats.Add(TEXT("STAT_GPUParticleTickTime"));
 	}
-	else
-	{
-		DesiredStats = StatsToRecord;
-	}
 	GeneratedStats.Empty();
 	for (int i = 0; i < DesiredStats.Num(); i++)
 	{
@@ -205,19 +251,22 @@ void FPerformanceMonitorModule::StartRecordingPerfTimers(FString FileNameToUse,
 		GeneratedStats.Add(DesiredStats[i], TempArray);
 	}
 	bRecording= true;
-	//AOrionPlayerController_Game* PlayerController = UOrionAutomationHelper::GetPlayerGameController();
-	//if (PlayerController)
-	//{
-	//	UWorld * PlayerWorld = PlayerController->GetWorld();
-	//	if (PlayerWorld)
-	//	{
-	//		// Postpone garbage collection for five minutes or until test ends, whichever comes first.
-	//		PlayerWorld->SetTimeUntilNextGarbageCollection(300.f);
-	//	}
-	//}
+	TimeOfTestStart = FPlatformTime::Seconds();
+	if (bRequiresCutsceneStart)
+	{
+		for (const FWorldContext& WorldContext : GEngine->GetWorldContexts())
+		{
+			if (WorldContext.WorldType == EWorldType::Game || WorldContext.WorldType == EWorldType::PIE)
+			{
+				UWorld* World = WorldContext.World();
+				GEngine->Exec(World, TEXT("ce start"));
+			}
+		}
+	}
+#endif
 }
 
-#endif
+
 void FPerformanceMonitorModule::RecordFrame()
 {
 	if (!bRecording)
@@ -225,50 +274,11 @@ void FPerformanceMonitorModule::RecordFrame()
 		return;
 	}
 	GetStatsBreakdown();
-	//#if STATS
-	//	FOrionCharacterPerfEvent::GetStatsBreakdown(PerfEvent);
-	//	AOrionPlayerController_Game* PlayerController = UOrionAutomationHelper::GetPlayerGameController();
-	//	AOrionCharHero * MyPawn = nullptr;
-	//	if (PlayerController != nullptr)
-	//	{
-	//		MyPawn = Cast<AOrionCharHero>(PlayerController->GetPawn());
-	//	}
-	//	// Ensure we have a player controller and a pawn
-	//	if ((PlayerController == nullptr) || (MyPawn == nullptr))
-	//	{
-	//		//we need these !!
-	//		return;
-	//	}
-	//
-	//	float CurrentPlatformTime = FPlatformTime::Seconds();
-	//	UWorld* CurWorld = PlayerController->GetWorld();
-	//	UGameViewportClient* Viewport = nullptr;
-	//	for (int i = 0; i < GEngine->GetWorldContexts().Num(); i++)
-	//	{
-	//		if (GEngine->GetWorldContexts()[i].WorldType == EWorldType::Game || GEngine->GetWorldContexts()[i].WorldType == EWorldType::PIE)
-	//		{
-	//			Viewport = GEngine->GetWorldContexts()[i].World()->GetGameViewport();
-	//			// If this viewport is actually valid, we break. Check to make sure we're not getting DS world by mistake.
-	//			if (Viewport)
-	//			{
-	//				break;
-	//			}
-	//		}
-	//	}
-	//	if (Viewport != nullptr)
-	//	{
-	//		const FStatUnitData* StatUnitData = Viewport->GetStatUnitData();
-	//
-	//			PerfEvent->OverallFrameTimes.FrameDurations.Add(StatUnitData->FrameTime);
-	//			PerfEvent->OverallFrameTimes.FrameGameTimes.Add(FGenericPlatformTime::ToMilliseconds(GGameThreadTime));
-	//			PerfEvent->OverallFrameTimes.FrameRenderTimes.Add(FGenericPlatformTime::ToMilliseconds(GRenderThreadTime));
-	//			PerfEvent->OverallFrameTimes.FrameTimeStamps.Add(CurrentPlatformTime);
-	//			PerfEvent->FrameNum++;
-	//	
-	//
-	//	}
-	//
-	//#endif
+	if (FPlatformTime::Seconds() - TestTimeOut > TimeOfTestStart)
+	{
+		StopRecordingPerformanceTimers();
+	}
+
 }
 
 void FPerformanceMonitorModule::GetStatsBreakdown()
@@ -321,18 +331,12 @@ void FPerformanceMonitorModule::StopRecordingPerformanceTimers()
 		FileToLogTo->Close();
 	}
 
-	// TODO: Force GC. 
-	//if (PlayerController)
-	//{
-	//	UWorld * PlayerWorld = PlayerController->GetWorld();
-	//	if (PlayerWorld)
-	//	{
-	//		// End our self imposed moratorium on GC.
-	//		PlayerWorld->ForceGarbageCollection();
-	//	}
-	//}
+
+	if (bExitOnCompletion)
+	{
+		FPlatformMisc::RequestExit(true);
+	}
 #endif
-	return;
 }
 
 void FPerformanceMonitorModule::RecordData()
@@ -340,35 +344,6 @@ void FPerformanceMonitorModule::RecordData()
 #if STATS
 
 
-	//TSharedPtr<FJsonObject> ReportJson;
-
-	//// FrameDurations
-	//EventToRecord.StatName = TEXT("FrameDuration");
-	//EventToRecord.MinVal = FMath::Min(PerfEvent->OverallFrameTimes.FrameDurations);
-	//EventToRecord.AvgVal = UOrionAutomationHelper::GetAverageOfArray(PerfEvent->OverallFrameTimes.FrameDurations);
-	//EventToRecord.MaxVal = FMath::Max(PerfEvent->OverallFrameTimes.FrameDurations);
-	//EventToRecord.NumFrames = PerfEvent->OverallFrameTimes.FrameDurations.Num();
-
-	//ReportJson = FJsonObjectConverter::UStructToJsonObject(EventToRecord);
-	//RecordPerfEvent();
-
-	//// GameThread
-	//EventToRecord.StatName = TEXT("GameThreadTime");
-	//EventToRecord.MinVal = FMath::Min(PerfEvent->OverallFrameTimes.FrameGameTimes);
-	//EventToRecord.AvgVal = UOrionAutomationHelper::GetAverageOfArray(PerfEvent->OverallFrameTimes.FrameGameTimes);
-	//EventToRecord.MaxVal = FMath::Max(PerfEvent->OverallFrameTimes.FrameGameTimes);
-	//EventToRecord.NumFrames = PerfEvent->OverallFrameTimes.FrameGameTimes.Num();
-
-	//ReportJson = FJsonObjectConverter::UStructToJsonObject(EventToRecord);
-	//RecordPerfEvent();
-
-
-	//// RenderThread
-	//EventToRecord.StatName = TEXT("RenderThreadTime");
-	//EventToRecord.MinVal = FMath::Min(PerfEvent->OverallFrameTimes.FrameRenderTimes);
-	//EventToRecord.AvgVal = UOrionAutomationHelper::GetAverageOfArray(PerfEvent->OverallFrameTimes.FrameRenderTimes);
-	//EventToRecord.MaxVal = FMath::Max(PerfEvent->OverallFrameTimes.FrameRenderTimes);
-	//EventToRecord.NumFrames = PerfEvent->OverallFrameTimes.FrameRenderTimes.Num();
 	FString StringToPrint = FString::Printf(TEXT("Interval (s),%0.4f\n"), TimeBetweenRecords);
 	FileToLogTo->Serialize(TCHAR_TO_ANSI(*StringToPrint), StringToPrint.Len());
 	for (int i = 0; i < StoredMessages.Num(); i++)
@@ -405,86 +380,7 @@ void FPerformanceMonitorModule::RecordData()
 				StringToPrint.Append(TEXT("\n"));
 				FileToLogTo->Serialize(TCHAR_TO_ANSI(*StringToPrint), StringToPrint.Len());
 	}
-	//for (TMap<FString, TMap<FString, TArray<float>>>::TConstIterator ThreadedItr(GeneratedStats); ThreadedItr; ++ThreadedItr)
-	//{
-	//	const TMap<FString, TArray<float>> ThreadGroup = ThreadedItr.Value();
-	//	for (TMap<FString, TArray<float>>::TConstIterator ThreadedItr2(ThreadGroup); ThreadedItr2; ++ThreadedItr2)
-	//	{
-	//		FString StringToPrint = ThreadedItr2.Key();
-	//		TArray<float> StatValues = ThreadedItr2.Value();
-	//		for (int i = 0; i < StatValues.Num(); i++)
-	//		{
-	//			StringToPrint = FString::Printf(TEXT("%s,%0.4f"), *StringToPrint, StatValues[i]);
-	//		}
-	//		StringToPrint.Append(TEXT("\n"));
-	//		FileToLogTo->Serialize(TCHAR_TO_ANSI(*StringToPrint), StringToPrint.Len());
-	//	}
-	//}
-	//for (TMap<FString, TMap<FString, TArray<FStatData>>>::TConstIterator ThreadedItr(PerfEvent->ThreadStatBreakdown); ThreadedItr; ++ThreadedItr)
-	//{
-	//	const FName ThreadName = ThreadedItr.Key();
-	//	const TMap<FName, TArray<FStatData>>& ThreadBreakdown = ThreadedItr.Value();
-	//	for (TMap<FName, TArray<FStatData>>::TConstIterator StatItr(ThreadBreakdown); StatItr; ++StatItr)
-	//	{
-	//		const FName StatName = StatItr.Key();
-	//		if (!DesiredStats.Contains(StatName.ToString()))
-	//		{
-	//			continue;
-	//		}
-	//		FString ThreadStatKeyAvg = FString::Printf(TEXT("AvgInclusive %s [%s]"), *ThreadName.ToString(), *StatName.ToString());
-	//		FString ThreadStatKeyMax = FString::Printf(TEXT("MaxInclusive %s [%s]"), *ThreadName.ToString(), *StatName.ToString());
-	//		FString StatAvgString;
-	//		FString StatMaxString;
-	//		TArray<float> StatAvgs;
-	//		TArray<float> StatMaxes;
-	//		const TArray<FStatData>& StatDatas = StatItr.Value();
-	//		int ValidFrame = 0;
-	//		for (int FrameNum = 0; FrameNum <= PerfEvent->FrameNum; ++FrameNum)
-	//		{
-	//			float AvgTime = -1.f;
-	//			float MaxTime = -1.f;
-	//			if (ValidFrame < StatDatas.Num())
-	//			{
-	//				const FStatData& StatData = StatDatas[ValidFrame];
-	//				if (StatData.FrameNum == FrameNum)
-	//				{
-	//					AvgTime = StatData.StatIncAvg;
-	//					MaxTime = StatData.StatIncMax;
-	//					++ValidFrame;
-	//				}
-	//			}
-
-	//			if (AvgTime != -1.f)
-	//			{
-	//				StatAvgs.Add(AvgTime);
-	//			}
-	//			if (MaxTime != -1.f)
-	//			{
-	//				StatMaxes.Add(MaxTime);
-	//			}
-	//		}
-	//		//float AvgAverage = UOrionAutomationHelper::GetAverageOfArray(StatAvgs);
-	//		//float AvgMaxes = UOrionAutomationHelper::GetAverageOfArray(StatMaxes);
-
-	//		//// StatMaxes
-	//		//EventToRecord.StatName = FString::Printf(TEXT("%s_MAX"), *StatName.ToString());
-	//		//EventToRecord.MinVal = FMath::Min(StatMaxes);
-	//		//EventToRecord.AvgVal = AvgMaxes;
-	//		//EventToRecord.MaxVal = FMath::Max(StatMaxes);
-	//		//EventToRecord.NumFrames = StatMaxes.Num();
-	//		//ReportJson = FJsonObjectConverter::UStructToJsonObject(EventToRecord);
-	//		//RecordPerfEvent();
-
-	//		//// StatAvgs
-	//		//EventToRecord.StatName = FString::Printf(TEXT("%s_AVG"), *StatName.ToString());
-	//		//EventToRecord.MinVal = FMath::Min(StatAvgs);
-	//		//EventToRecord.AvgVal = AvgAverage;
-	//		//EventToRecord.MaxVal = FMath::Max(StatAvgs);
-	//		//EventToRecord.NumFrames = StatAvgs.Num();
-	//		//ReportJson = FJsonObjectConverter::UStructToJsonObject(EventToRecord);
-	//		//RecordPerfEvent();
-	//	}
-	//}
+
 #endif
 }
 
diff --git a/Engine/Plugins/Performance/PerformanceMonitor/Source/PerformanceMonitor/Private/PerformanceMonitor.h b/Engine/Plugins/Performance/PerformanceMonitor/Source/PerformanceMonitor/Private/PerformanceMonitor.h
--- a/Engine/Plugins/Performance/PerformanceMonitor/Source/PerformanceMonitor/Private/PerformanceMonitor.h
+++ b/Engine/Plugins/Performance/PerformanceMonitor/Source/PerformanceMonitor/Private/PerformanceMonitor.h
@@ -31,9 +31,14 @@ private:
 	FDelegateHandle TickHandler;
 
 	bool bRecording;
+	float TimeOfTestStart;
+	float TestTimeOut;
 	float TimeOfLastRecord;
 	float TimeBetweenRecords;
+	FString MapToTest;
 	bool bHasWarnedAboutTime;
+	bool bExitOnCompletion;
+	bool bRequiresCutsceneStart;
 
 protected:
 public:
diff --git a/Engine/Shaders/QuadOverdraw.usf b/Engine/Shaders/QuadOverdraw.usf
--- a/Engine/Shaders/QuadOverdraw.usf
+++ b/Engine/Shaders/QuadOverdraw.usf
@@ -47,10 +47,16 @@ float ComplexityToFloat(uint Complexity)
 
 #if OUTPUT_QUAD_OVERDRAW
 
+#if FORWARD_SHADING
+#define QUAD_BUFFER_REGISTER u1
+#else
+#define QUAD_BUFFER_REGISTER u7
+#endif
+
 // The temporary buffer used to synchronize and exchange data between quad sub-pixels.
 // Left half hold QuadDescriptor, right half hold QuadComplexity
 // Both are halfres here.
-RWTexture2D<uint> RWQuadBuffer : register(u7);
+RWTexture2D<uint> RWQuadBuffer : register(QUAD_BUFFER_REGISTER);
 
 
 // The QuadComplexity is the same resource as the QuadDescriptor with an offset in X.
diff --git a/Engine/Source/Runtime/Apple/MetalRHI/Private/MetalStateCache.cpp b/Engine/Source/Runtime/Apple/MetalRHI/Private/MetalStateCache.cpp
--- a/Engine/Source/Runtime/Apple/MetalRHI/Private/MetalStateCache.cpp
+++ b/Engine/Source/Runtime/Apple/MetalRHI/Private/MetalStateCache.cpp
@@ -1116,6 +1116,7 @@ void FMetalStateCache::SetShaderResourceView(EShaderFrequency ShaderStage, uint3
 		FRHITexture* Texture = SRV->SourceTexture.GetReference();
 		FMetalVertexBuffer* VB = SRV->SourceVertexBuffer.GetReference();
 		FMetalIndexBuffer* IB = SRV->SourceIndexBuffer.GetReference();
+		FMetalStructuredBuffer* SB = SRV->SourceStructuredBuffer.GetReference();
 		if (Texture)
 		{
 			FMetalSurface* Surface = SRV->TextureView;
@@ -1137,6 +1138,10 @@ void FMetalStateCache::SetShaderResourceView(EShaderFrequency ShaderStage, uint3
 		{
 			SetShaderBuffer(ShaderStage, IB->Buffer, nil, 0, IB->GetSize(), BindIndex);
 		}
+		else if (SB)
+		{
+			SetShaderBuffer(ShaderStage, SB->Buffer, nil, 0, SB->GetSize(), BindIndex);
+		}
 	}
 	FShaderCache::SetSRV(ShaderStage, BindIndex, SRV);
 }
diff --git a/Engine/Source/Runtime/Apple/MetalRHI/Private/MetalUAV.cpp b/Engine/Source/Runtime/Apple/MetalRHI/Private/MetalUAV.cpp
--- a/Engine/Source/Runtime/Apple/MetalRHI/Private/MetalUAV.cpp
+++ b/Engine/Source/Runtime/Apple/MetalRHI/Private/MetalUAV.cpp
@@ -183,7 +183,6 @@ FShaderResourceViewRHIRef FMetalDynamicRHI::RHICreateShaderResourceView_RenderTh
 
 FShaderResourceViewRHIRef FMetalDynamicRHI::RHICreateShaderResourceView_RenderThread(class FRHICommandListImmediate& RHICmdList, FStructuredBufferRHIParamRef StructuredBuffer)
 {
-	UE_LOG(LogRHI, Fatal,TEXT("Metal RHI doesn't support RHICreateShaderResourceView with FStructuredBufferRHIParamRef yet!"));
 	return GDynamicRHI->RHICreateShaderResourceView(StructuredBuffer);
 }
 
@@ -193,8 +192,13 @@ FShaderResourceViewRHIRef FMetalDynamicRHI::RHICreateShaderResourceView(FStructu
 	FMetalStructuredBuffer* StructuredBuffer = ResourceCast(StructuredBufferRHI);
 
 	FMetalShaderResourceView* SRV = new FMetalShaderResourceView;
+	SRV->SourceVertexBuffer = nullptr;
+	SRV->SourceIndexBuffer = nullptr;
 	SRV->TextureView = nullptr;
-	UE_LOG(LogRHI, Fatal,TEXT("Metal RHI doesn't support RHICreateShaderResourceView with FStructuredBufferRHIParamRef yet!"));
+	SRV->SourceStructuredBuffer = StructuredBuffer;
+	
+	UE_LOG(LogRHI, Warning, TEXT("FShaderCache doesn't support RHICreateShaderResourceView with FStructuredBufferRHIParamRef yet!"));
+	
 	return SRV;
 }
 
@@ -205,6 +209,8 @@ FShaderResourceViewRHIRef FMetalDynamicRHI::RHICreateShaderResourceView(FVertexB
 	FMetalShaderResourceView* SRV = new FMetalShaderResourceView;
 	SRV->SourceVertexBuffer = VertexBuffer;
 	SRV->TextureView = nullptr;
+	SRV->SourceIndexBuffer = nullptr;
+	SRV->SourceStructuredBuffer = nullptr;
 	
 	FShaderCache::LogSRV(SRV, VertexBufferRHI, Stride, Format);
 	
@@ -219,6 +225,7 @@ FShaderResourceViewRHIRef FMetalDynamicRHI::RHICreateShaderResourceView(FIndexBu
 	SRV->SourceVertexBuffer = nullptr;
 	SRV->SourceIndexBuffer = Buffer;
 	SRV->TextureView = nullptr;
+	SRV->SourceStructuredBuffer = nullptr;
 	
 	UE_LOG(LogRHI, Warning, TEXT("FShaderCache doesn't support RHICreateShaderResourceView with FIndexBufferRHIParamRef yet!"));
 	//FShaderCache::LogSRV(SRV, VertexBufferRHI, Stride, Format);
@@ -234,6 +241,10 @@ FShaderResourceViewRHIRef FMetalDynamicRHI::RHICreateShaderResourceView(FTexture
 	FMetalSurface* Surface = GetMetalSurfaceFromRHITexture(Texture2DRHI);
 	SRV->TextureView = Surface ? new FMetalSurface(*Surface, NSMakeRange(MipLevel, 1)) : nullptr;
 	
+	SRV->SourceVertexBuffer = nullptr;
+	SRV->SourceIndexBuffer = nullptr;
+	SRV->SourceStructuredBuffer = nullptr;
+	
 	SRV->MipLevel = MipLevel;
 	SRV->NumMips = 1;
 	SRV->Format = PF_Unknown;
@@ -251,6 +262,10 @@ FShaderResourceViewRHIRef FMetalDynamicRHI::RHICreateShaderResourceView(FTexture
 	FMetalSurface* Surface = GetMetalSurfaceFromRHITexture(Texture2DRHI);
 	SRV->TextureView = Surface ? new FMetalSurface(*Surface, NSMakeRange(MipLevel, NumMipLevels), (EPixelFormat)Format) : nullptr;
 	
+	SRV->SourceVertexBuffer = nullptr;
+	SRV->SourceIndexBuffer = nullptr;
+	SRV->SourceStructuredBuffer = nullptr;
+	
 	SRV->MipLevel = MipLevel;
 	SRV->NumMips = NumMipLevels;
 	SRV->Format = Format;
@@ -267,6 +282,10 @@ FShaderResourceViewRHIRef FMetalDynamicRHI::RHICreateShaderResourceView(FTexture
 	FMetalSurface* Surface = GetMetalSurfaceFromRHITexture(Texture3DRHI);
 	SRV->TextureView = Surface ? new FMetalSurface(*Surface, NSMakeRange(MipLevel, 1)) : nullptr;
 	
+	SRV->SourceVertexBuffer = nullptr;
+	SRV->SourceIndexBuffer = nullptr;
+	SRV->SourceStructuredBuffer = nullptr;
+	
 	SRV->MipLevel = MipLevel;
 	SRV->NumMips = 1;
 	SRV->Format = PF_Unknown;
@@ -283,6 +302,10 @@ FShaderResourceViewRHIRef FMetalDynamicRHI::RHICreateShaderResourceView(FTexture
 	FMetalSurface* Surface = GetMetalSurfaceFromRHITexture(Texture2DArrayRHI);
 	SRV->TextureView = Surface ? new FMetalSurface(*Surface, NSMakeRange(MipLevel, 1)) : nullptr;
 	
+	SRV->SourceVertexBuffer = nullptr;
+	SRV->SourceIndexBuffer = nullptr;
+	SRV->SourceStructuredBuffer = nullptr;
+	
 	SRV->MipLevel = MipLevel;
 	SRV->NumMips = 1;
 	SRV->Format = PF_Unknown;
@@ -299,6 +322,10 @@ FShaderResourceViewRHIRef FMetalDynamicRHI::RHICreateShaderResourceView(FTexture
 	FMetalSurface* Surface = GetMetalSurfaceFromRHITexture(TextureCubeRHI);
 	SRV->TextureView = Surface ? new FMetalSurface(*Surface, NSMakeRange(MipLevel, 1)) : nullptr;
 	
+	SRV->SourceVertexBuffer = nullptr;
+	SRV->SourceIndexBuffer = nullptr;
+	SRV->SourceStructuredBuffer = nullptr;
+	
 	SRV->MipLevel = MipLevel;
 	SRV->NumMips = 1;
 	SRV->Format = PF_Unknown;
diff --git a/Engine/Source/Runtime/Apple/MetalRHI/Public/MetalResources.h b/Engine/Source/Runtime/Apple/MetalRHI/Public/MetalResources.h
--- a/Engine/Source/Runtime/Apple/MetalRHI/Public/MetalResources.h
+++ b/Engine/Source/Runtime/Apple/MetalRHI/Public/MetalResources.h
@@ -776,6 +776,9 @@ public:
 	// The texture that this SRV come from
 	TRefCountPtr<FRHITexture> SourceTexture;
 	
+	// The source structured buffer (can be null)
+	TRefCountPtr<FMetalStructuredBuffer> SourceStructuredBuffer;
+	
 	FMetalSurface* TextureView;
 	uint8 MipLevel;
 	uint8 NumMips;
diff --git a/Engine/Source/Runtime/D3D12RHI/D3D12RHI.Build.cs b/Engine/Source/Runtime/D3D12RHI/D3D12RHI.Build.cs
--- a/Engine/Source/Runtime/D3D12RHI/D3D12RHI.Build.cs
+++ b/Engine/Source/Runtime/D3D12RHI/D3D12RHI.Build.cs
@@ -24,34 +24,6 @@ public class D3D12RHI : ModuleRules
 			PrivateIncludePathModuleNames.AddRange(new string[] { "TaskGraph" });
 		}
 
-		///////////////////////////////////////////////////////////////
-		// Platform agnostic defines
-		///////////////////////////////////////////////////////////////
-		Definitions.Add("SUB_ALLOCATED_DEFAULT_ALLOCATIONS=1");
-
-		Definitions.Add("DEBUG_RESOURCE_STATES=0");
-
-		// DX12 doesn't support higher MSAA count
-		Definitions.Add("DX_MAX_MSAA_COUNT=8");
-
-		// This is a value that should be tweaked to fit the app, lower numbers will have better performance
-		Definitions.Add("MAX_SRVS=32");
-		Definitions.Add("MAX_SAMPLERS=16");
-		Definitions.Add("MAX_UAVS=8");
-		Definitions.Add("MAX_CBS=8");
-
-		// This value controls how many root constant buffers can be used per shader stage in a root signature.
-		// Note: Using root descriptors significantly increases the size of root signatures (each root descriptor is 2 DWORDs).
-		Definitions.Add("MAX_ROOT_CBVS=MAX_CBS");
-
-		// PC uses a multiple root signatures that are optimized for a particular draw/dispatch.
-		Definitions.Add("USE_STATIC_ROOT_SIGNATURE=0");
-
-		// How many residency packets can be in flight before the rendering thread
-		// blocks for them to drain. Should be ~ NumBufferedFrames * AvgNumSubmissionsPerFrame i.e.
-		// enough to ensure that the GPU is rarely blocked by residency work
-		Definitions.Add("RESIDENCY_PIPELINE_DEPTH=6");
-
 		///////////////////////////////////////////////////////////////
 		// Platform specific defines
 		///////////////////////////////////////////////////////////////
@@ -66,38 +38,6 @@ public class D3D12RHI : ModuleRules
 			AddEngineThirdPartyPrivateStaticDependencies(Target, "DX12");
 			AddEngineThirdPartyPrivateStaticDependencies(Target, "DX11");
 			AddEngineThirdPartyPrivateStaticDependencies(Target, "NVAPI");
-
-			Definitions.Add("ENABLE_RESIDENCY_MANAGEMENT=1");
-
-			Definitions.Add("ASYNC_DEFERRED_DELETION=1");
-
-			Definitions.Add("PLATFORM_SUPPORTS_MGPU=1");
-
-			Definitions.Add("PIPELINE_STATE_FILE_LOCATION=FPaths::GameSavedDir()");
-
-			Definitions.Add("USE_PIX=D3D12_PROFILING_ENABLED");
-		}
-		else if (Target.Platform == UnrealTargetPlatform.XboxOne)
-		{
-			Definitions.Add("ENABLE_RESIDENCY_MANAGEMENT=0");
-
-			Definitions.Add("ASYNC_DEFERRED_DELETION=0");
-
-			Definitions.Add("PLATFORM_SUPPORTS_MGPU=0");
-
-			Definitions.Add("ASYNC_DEFERRED_DELETION=0");
-
-			// The number of sampler descriptors with the maximum value of 2048
-			// If the heap type is unbounded the number could be increased to avoid rollovers.
-			Definitions.Add("ENABLE_UNBOUNDED_SAMPLER_DESCRIPTORS=0");
-			Definitions.Add("NUM_SAMPLER_DESCRIPTORS=D3D12_MAX_SHADER_VISIBLE_SAMPLER_HEAP_SIZE");
-
-			// Xbox doesn't have DXGI but common code needs this defined for headers
-			Definitions.Add("DXGI_QUERY_VIDEO_MEMORY_INFO=int");
-
-			Definitions.Add("PIPELINE_STATE_FILE_LOCATION=FPaths::GameContentDir()");
-
-			Definitions.Add("USE_PIX=XBOXONE_PROFILING_ENABLED");
 		}
 	}
 }
diff --git a/Engine/Source/Runtime/D3D12RHI/Public/D3D12RHI.h b/Engine/Source/Runtime/D3D12RHI/Public/D3D12RHI.h
--- a/Engine/Source/Runtime/D3D12RHI/Public/D3D12RHI.h
+++ b/Engine/Source/Runtime/D3D12RHI/Public/D3D12RHI.h
@@ -8,10 +8,65 @@
 
 DECLARE_STATS_GROUP(TEXT("D3D12RHI"), STATGROUP_D3D12RHI, STATCAT_Advanced);
 
+
+///////////////////////////////////////////////////////////////
+// Platform agnostic defines
+///////////////////////////////////////////////////////////////
+#define SUB_ALLOCATED_DEFAULT_ALLOCATIONS	1
+
+#define DEBUG_RESOURCE_STATES	0
+
+// DX12 doesn't support higher MSAA count
+#define DX_MAX_MSAA_COUNT	8
+
+// This is a value that should be tweaked to fit the app, lower numbers will have better performance
+#define MAX_SRVS		32
+#define MAX_SAMPLERS	16
+#define MAX_UAVS		8
+#define MAX_CBS			16
+
+// This value controls how many root constant buffers can be used per shader stage in a root signature.
+// Note: Using root descriptors significantly increases the size of root signatures (each root descriptor is 2 DWORDs).
+#define MAX_ROOT_CBVS	MAX_CBS
+
+// PC uses a multiple root signatures that are optimized for a particular draw/dispatch.
+#define USE_STATIC_ROOT_SIGNATURE	0
+
+// How many residency packets can be in flight before the rendering thread
+// blocks for them to drain. Should be ~ NumBufferedFrames * AvgNumSubmissionsPerFrame i.e.
+// enough to ensure that the GPU is rarely blocked by residency work
+#define RESIDENCY_PIPELINE_DEPTH	6
+
+#if PLATFORM_XBOXONE
+	#define ENABLE_RESIDENCY_MANAGEMENT				0
+	#define ASYNC_DEFERRED_DELETION					0
+	#define PLATFORM_SUPPORTS_MGPU					0
+	#define ASYNC_DEFERRED_DELETION					0
+	#define PIPELINE_STATE_FILE_LOCATION			FPaths::GameContentDir()
+	#define USE_PIX									XBOXONE_PROFILING_ENABLED
+
+	// The number of sampler descriptors with the maximum value of 2048
+	// If the heap type is unbounded the number could be increased to avoid rollovers.
+	#define ENABLE_UNBOUNDED_SAMPLER_DESCRIPTORS	0
+	#define NUM_SAMPLER_DESCRIPTORS					D3D12_MAX_SHADER_VISIBLE_SAMPLER_HEAP_SIZE
+
+	// Xbox doesn't have DXGI but common code needs this defined for headers
+	#define DXGI_QUERY_VIDEO_MEMORY_INFO			int
+#elif PLATFORM_WINDOWS
+	#define ENABLE_RESIDENCY_MANAGEMENT				1
+	#define ASYNC_DEFERRED_DELETION					1
+	#define PLATFORM_SUPPORTS_MGPU					1
+	#define PIPELINE_STATE_FILE_LOCATION			FPaths::GameSavedDir()
+	#define USE_PIX									D3D12_PROFILING_ENABLED
+#else
+	#error	Unknown platform!
+#endif
+
+
 #define FD3D12_TEXTURE_DATA_PITCH_ALIGNMENT D3D12_TEXTURE_DATA_PITCH_ALIGNMENT
 #define FD3D12_DESCRIPTOR_HEAP_TYPE_SAMPLER D3D12_DESCRIPTOR_HEAP_TYPE_SAMPLER
 
-typedef uint8 CBVSlotMask;
+typedef uint16 CBVSlotMask;
 static_assert(MAX_ROOT_CBVS <= MAX_CBS, "MAX_ROOT_CBVS must be <= MAX_CBS.");
 static_assert((8 * sizeof(CBVSlotMask)) >= MAX_CBS, "CBVSlotMask isn't large enough to cover all CBs. Please increase the size.");
 static const CBVSlotMask GRootCBVSlotMask = (1 << MAX_ROOT_CBVS) - 1; // Mask for all slots that are used by root descriptors.
diff --git a/Engine/Source/Runtime/Engine/Classes/Components/SceneCaptureComponent.h b/Engine/Source/Runtime/Engine/Classes/Components/SceneCaptureComponent.h
--- a/Engine/Source/Runtime/Engine/Classes/Components/SceneCaptureComponent.h
+++ b/Engine/Source/Runtime/Engine/Classes/Components/SceneCaptureComponent.h
@@ -66,6 +66,10 @@ class ENGINE_API USceneCaptureComponent : public USceneComponent
 	/** Whether to update the capture's contents on movement.  Disable if you are going to capture manually from blueprint. */
 	UPROPERTY(EditAnywhere, BlueprintReadWrite, Category=SceneCapture)
 	bool bCaptureOnMovement;
+	
+	/** Scales the distance used by LOD. Set to values greater than 1 to cause the scene capture to use lower LODs than the main view to speed up the scene capture pass. */
+	UPROPERTY(EditAnywhere, Category=PlanarReflection, meta=(UIMin = ".1", UIMax = "10"), AdvancedDisplay)
+	float LODDistanceFactor;
 
 	/** if > 0, sets a maximum render distance override.  Can be used to cull distant objects from a reflection if the reflecting plane is in an enclosed area like a hallway or room */
 	UPROPERTY(EditAnywhere, BlueprintReadWrite, Category=SceneCapture, meta=(UIMin = "100", UIMax = "10000"))
diff --git a/Engine/Source/Runtime/Engine/Private/BatchedElements.cpp b/Engine/Source/Runtime/Engine/Private/BatchedElements.cpp
--- a/Engine/Source/Runtime/Engine/Private/BatchedElements.cpp
+++ b/Engine/Source/Runtime/Engine/Private/BatchedElements.cpp
@@ -340,7 +340,9 @@ static void SetBlendState(FRHICommandList& RHICmdList, ESimpleElementBlendMode B
 	static const auto CVarCompositeMode = IConsoleManager::Get().FindTConsoleVariableDataInt(TEXT("r.HDR.UI.CompositeMode"));
 	static const auto CVarHDROutputEnabled = IConsoleManager::Get().FindTConsoleVariableDataInt(TEXT("r.HDR.EnableHDROutput"));
 
-	const bool bCompositeUI = CVarCompositeMode->GetValueOnRenderThread() != 0 && GRHISupportsHDROutput && CVarHDROutputEnabled->GetValueOnRenderThread() != 0;
+	const bool bCompositeUI = GRHISupportsHDROutput
+		&& CVarCompositeMode && CVarCompositeMode->GetValueOnRenderThread() != 0 
+		&& CVarHDROutputEnabled && CVarHDROutputEnabled->GetValueOnRenderThread() != 0;
 
 	if (bCompositeUI)
 	{
diff --git a/Engine/Source/Runtime/Engine/Private/Components/SceneCaptureComponent.cpp b/Engine/Source/Runtime/Engine/Private/Components/SceneCaptureComponent.cpp
--- a/Engine/Source/Runtime/Engine/Private/Components/SceneCaptureComponent.cpp
+++ b/Engine/Source/Runtime/Engine/Private/Components/SceneCaptureComponent.cpp
@@ -169,6 +169,7 @@ USceneCaptureComponent::USceneCaptureComponent(const FObjectInitializer& ObjectI
 {
 	bCaptureEveryFrame = true;
 	bCaptureOnMovement = true;
+	LODDistanceFactor = 1.0f;
 	MaxViewDistanceOverride = -1;
 
 	// Disable features that are not desired when capturing the scene
diff --git a/Engine/Source/Runtime/Engine/Private/GPUSkinCache.cpp b/Engine/Source/Runtime/Engine/Private/GPUSkinCache.cpp
--- a/Engine/Source/Runtime/Engine/Private/GPUSkinCache.cpp
+++ b/Engine/Source/Runtime/Engine/Private/GPUSkinCache.cpp
@@ -68,7 +68,7 @@ TAutoConsoleVariable<int32> CVarGPUSkinCacheRecomputeTangents(
 	TEXT("This option enables recomputing the vertex tangents on the GPU.\n")
 	TEXT("Can be changed at runtime, requires both r.SkinCache.CompileShaders=1 and r.SkinCache.Mode=1\n")
 	TEXT(" 0: off\n")
-	TEXT(" 1: on, forces all skinned object to Tecompute Tangents\n")
+	TEXT(" 1: on, forces all skinned object to Recompute Tangents\n")
 	TEXT(" 2: on, only recompute tangents on skinned objects who ticked the Tecompute Tangents checkbox(default)\n"),
 	ECVF_RenderThreadSafe
 	);
diff --git a/Engine/Source/Runtime/Engine/Private/Materials/MaterialExpressions.cpp b/Engine/Source/Runtime/Engine/Private/Materials/MaterialExpressions.cpp
--- a/Engine/Source/Runtime/Engine/Private/Materials/MaterialExpressions.cpp
+++ b/Engine/Source/Runtime/Engine/Private/Materials/MaterialExpressions.cpp
@@ -6631,7 +6631,7 @@ int32 UMaterialExpressionSceneTexture::Compile(class FMaterialCompiler* Compiler
 		return Compiler->SceneTextureSize(SceneTextureId, OutputIndex == 2);
 	}
 
-	return Compiler->Errorf(TEXT("Invalid intput parameter"));
+	return Compiler->Errorf(TEXT("Invalid input parameter"));
 }
 
 void UMaterialExpressionSceneTexture::GetCaption(TArray<FString>& OutCaptions) const
diff --git a/Engine/Source/Runtime/Engine/Private/Particles/ParticleTrail2EmitterInstance.cpp b/Engine/Source/Runtime/Engine/Private/Particles/ParticleTrail2EmitterInstance.cpp
--- a/Engine/Source/Runtime/Engine/Private/Particles/ParticleTrail2EmitterInstance.cpp
+++ b/Engine/Source/Runtime/Engine/Private/Particles/ParticleTrail2EmitterInstance.cpp
@@ -2215,27 +2215,38 @@ bool FParticleRibbonEmitterInstance::ResolveSourcePoint(int32 InTrailIdx,
 								break;
 							case EPSSM_Sequential:
 							{
-								bool bDone = false;
-
-								if (++LastSelectedParticleIndex >= SourceEmitter->ActiveParticles)
+								bool bInUse = false;
+								LastSelectedParticleIndex ++;
+								if (LastSelectedParticleIndex >= SourceEmitter->ActiveParticles)
 								{
 									LastSelectedParticleIndex = 0;
 								}
-								Index = LastSelectedParticleIndex;
-							}
-							break;
-						}
 
-						for (int32 TrailCheckIdx = 0; TrailCheckIdx < MaxTrailCount; TrailCheckIdx++)
-						{
-							if (TrailCheckIdx != InTrailIdx)
-							{
-								if (SourceIndices[TrailCheckIdx] == SourceEmitter->ParticleIndices[Index])
+								// need to check if the next source index is in use, and go through until we find one that isn't
+								do 
+								{
+									bInUse = false;
+									for (int32 TrailCheckIdx = 0; TrailCheckIdx < MaxTrailCount; TrailCheckIdx++)
+									{
+										if (TrailCheckIdx != InTrailIdx && SourceIndices[TrailCheckIdx] == SourceEmitter->ParticleIndices[LastSelectedParticleIndex])
+										{
+											bInUse = true;
+										}
+									}
+									if (bInUse == true)
+									{
+										LastSelectedParticleIndex++;
+									}								
+								} while (bInUse && LastSelectedParticleIndex<SourceEmitter->ActiveParticles);
+
+								if (LastSelectedParticleIndex >= SourceEmitter->ActiveParticles)
 								{
-									Index = -1;
-									break;
+									LastSelectedParticleIndex = -1;
 								}
+
+								Index = LastSelectedParticleIndex;
 							}
+							break;
 						}
 
 						SourceIndices[InTrailIdx] = Index != -1 ? SourceEmitter->ParticleIndices[Index] : -1;
diff --git a/Engine/Source/Runtime/Engine/Private/SceneManagement.cpp b/Engine/Source/Runtime/Engine/Private/SceneManagement.cpp
--- a/Engine/Source/Runtime/Engine/Private/SceneManagement.cpp
+++ b/Engine/Source/Runtime/Engine/Private/SceneManagement.cpp
@@ -395,7 +395,7 @@ int8 ComputeStaticMeshLOD( const FStaticMeshRenderData* RenderData, const FVecto
 	{
 		const int32 NumLODs = MAX_STATIC_MESH_LODS;
 
-		const float ScreenRadiusSquared = ComputeBoundsScreenRadiusSquared(Origin, SphereRadius, View) * FactorScale * FactorScale;
+		const float ScreenRadiusSquared = ComputeBoundsScreenRadiusSquared(Origin, SphereRadius, View) * FactorScale * FactorScale * View.LODDistanceFactor * View.LODDistanceFactor;
 
 		// Walk backwards and return the first matching LOD
 		for (int32 LODIndex = NumLODs - 1; LODIndex >= 0; --LODIndex)
diff --git a/Engine/Source/Runtime/Engine/Private/UnrealEngine.cpp b/Engine/Source/Runtime/Engine/Private/UnrealEngine.cpp
--- a/Engine/Source/Runtime/Engine/Private/UnrealEngine.cpp
+++ b/Engine/Source/Runtime/Engine/Private/UnrealEngine.cpp
@@ -386,7 +386,7 @@ void SystemResolutionSinkCallback()
 	int32 WindowModeInt = GSystemResolution.WindowMode;
 
 	static const auto CVarHDROutputEnabled = IConsoleManager::Get().FindTConsoleVariableDataInt(TEXT("r.HDR.EnableHDROutput"));
-	bool bHDROutputEnabled = GRHISupportsHDROutput && CVarHDROutputEnabled->GetValueOnAnyThread() != 0;
+	bool bHDROutputEnabled = GRHISupportsHDROutput && CVarHDROutputEnabled && CVarHDROutputEnabled->GetValueOnAnyThread() != 0;
 	
 	if (FParse::Resolution(*ResString, ResX, ResY, WindowModeInt))
 	{
diff --git a/Engine/Source/Runtime/Renderer/Private/BasePassRendering.cpp b/Engine/Source/Runtime/Renderer/Private/BasePassRendering.cpp
--- a/Engine/Source/Runtime/Renderer/Private/BasePassRendering.cpp
+++ b/Engine/Source/Runtime/Renderer/Private/BasePassRendering.cpp
@@ -331,7 +331,7 @@ static void SetDepthStencilStateForBasePass(FRHICommandList& RHICmdList, FDrawin
 				0xFF, GET_STENCIL_BIT_MASK(RECEIVE_DECAL, 1) | STENCIL_LIGHTING_CHANNELS_MASK(0x7)
 			>::GetRHI(), StencilValue);
 		}
-		else
+		else if (DrawRenderState.GetDepthStencilAccess() & FExclusiveDepthStencil::DepthWrite)
 		{
 			DrawRenderState.SetDepthStencilState(RHICmdList, TStaticDepthStencilState<
 				true, CF_GreaterEqual,
@@ -340,6 +340,15 @@ static void SetDepthStencilStateForBasePass(FRHICommandList& RHICmdList, FDrawin
 				0xFF, GET_STENCIL_BIT_MASK(RECEIVE_DECAL, 1) | STENCIL_LIGHTING_CHANNELS_MASK(0x7)
 			>::GetRHI(), StencilValue);
 		}
+		else
+		{
+			DrawRenderState.SetDepthStencilState(RHICmdList, TStaticDepthStencilState<
+				false, CF_GreaterEqual,
+				true, CF_Always, SO_Keep, SO_Keep, SO_Replace,
+				false, CF_Always, SO_Keep, SO_Keep, SO_Keep,
+				0xFF, GET_STENCIL_BIT_MASK(RECEIVE_DECAL, 1) | STENCIL_LIGHTING_CHANNELS_MASK(0x7)
+			>::GetRHI(), StencilValue);
+		}
 	}
 	else if (bMaskInEarlyPass)
 	{
@@ -821,7 +830,7 @@ void GetBasePassShaders<FUniformLightMapPolicy>(
  * Renders the scene's base pass 
  * @return true if anything was rendered
  */
-bool FDeferredShadingSceneRenderer::RenderBasePass(FRHICommandListImmediate& RHICmdList)
+bool FDeferredShadingSceneRenderer::RenderBasePass(FRHICommandListImmediate& RHICmdList, FExclusiveDepthStencil::Type BasePassDepthStencilAccess)
 {
 	bool bDirty = false;
 
@@ -845,10 +854,10 @@ bool FDeferredShadingSceneRenderer::RenderBasePass(FRHICommandListImmediate& RHI
 				FViewInfo& View = Views[ViewIndex];
 				if (View.ShouldRenderView())
 				{
-					RenderBasePassViewParallel(View, RHICmdList);
+					RenderBasePassViewParallel(View, RHICmdList, BasePassDepthStencilAccess);
 				}
 
-				RenderEditorPrimitives(RHICmdList, View, bDirty);
+				RenderEditorPrimitives(RHICmdList, View, BasePassDepthStencilAccess, bDirty);
 			}
 
 			bDirty = true; // assume dirty since we are not going to wait
@@ -861,10 +870,10 @@ bool FDeferredShadingSceneRenderer::RenderBasePass(FRHICommandListImmediate& RHI
 				FViewInfo& View = Views[ViewIndex];
 				if (View.ShouldRenderView())
 				{
-					bDirty |= RenderBasePassView(RHICmdList, View);
+					bDirty |= RenderBasePassView(RHICmdList, View, BasePassDepthStencilAccess);
 				}
 
-				RenderEditorPrimitives(RHICmdList, View, bDirty);
+				RenderEditorPrimitives(RHICmdList, View, BasePassDepthStencilAccess, bDirty);
 			}
 		}	
 	}
@@ -1103,8 +1112,10 @@ void FDeferredShadingSceneRenderer::RenderBasePassDynamicDataParallel(FParallelC
 	ParallelCommandListSet.AddParallelCommandList(CmdList, AnyThreadCompletionEvent);
 }
 
-static void SetupBasePassView(FRHICommandList& RHICmdList, const FViewInfo& View, FDrawingPolicyRenderState& DrawRenderState, const bool bShaderComplexity, const bool bIsEditorPrimitivePass = false)
+static void SetupBasePassView(FRHICommandList& RHICmdList, const FViewInfo& View, FDrawingPolicyRenderState& DrawRenderState, FExclusiveDepthStencil::Type BasePassDepthStencilAccess, const bool bShaderComplexity, const bool bIsEditorPrimitivePass = false)
 {
+	DrawRenderState.SetDepthStencilAccess(BasePassDepthStencilAccess);
+
 	if (bShaderComplexity)
 	{
 		// Additive blending when shader complexity viewmode is enabled.
@@ -1125,9 +1136,15 @@ static void SetupBasePassView(FRHICommandList& RHICmdList, const FViewInfo& View
 			DrawRenderState.SetBlendState(RHICmdList, TStaticBlendStateWriteMask<CW_RGBA, CW_RGBA, CW_RGBA, CW_RGBA>::GetRHI());
 		}
 
+		if (DrawRenderState.GetDepthStencilAccess() & FExclusiveDepthStencil::DepthWrite)
+		{
 			DrawRenderState.SetDepthStencilState(RHICmdList, TStaticDepthStencilState<true, CF_DepthNearOrEqual>::GetRHI());
 		}
-	RHICmdList.SetScissorRect(false, 0, 0, 0, 0);
+		else
+		{
+			DrawRenderState.SetDepthStencilState(RHICmdList, TStaticDepthStencilState<false, CF_DepthNearOrEqual>::GetRHI());
+		}
+	}
 
 	if (!View.IsInstancedStereoPass() || bIsEditorPrimitivePass)
 	{
@@ -1156,10 +1173,12 @@ class FBasePassParallelCommandListSet : public FParallelCommandListSet
 {
 public:
 	const FSceneViewFamily& ViewFamily;
+	FExclusiveDepthStencil::Type BasePassDepthStencilAccess;
 
-	FBasePassParallelCommandListSet(const FViewInfo& InView, FRHICommandListImmediate& InParentCmdList, bool bInParallelExecute, bool bInCreateSceneContext, const FSceneViewFamily& InViewFamily)
+	FBasePassParallelCommandListSet(const FViewInfo& InView, FRHICommandListImmediate& InParentCmdList, bool bInParallelExecute, bool bInCreateSceneContext, const FSceneViewFamily& InViewFamily, FExclusiveDepthStencil::Type InBasePassDepthStencilAccess)
 		: FParallelCommandListSet(GET_STATID(STAT_CLP_Basepass), InView, InParentCmdList, bInParallelExecute, bInCreateSceneContext)
 		, ViewFamily(InViewFamily)
+		, BasePassDepthStencilAccess(InBasePassDepthStencilAccess)
 	{
 		SetStateOnCommandList(ParentCmdList);
 	}
@@ -1172,26 +1191,27 @@ public:
 	virtual void SetStateOnCommandList(FRHICommandList& CmdList) override
 	{
 		FParallelCommandListSet::SetStateOnCommandList(CmdList);
-		FSceneRenderTargets::Get(CmdList).BeginRenderingGBuffer(CmdList, ERenderTargetLoadAction::ELoad, ERenderTargetLoadAction::ELoad, ViewFamily.EngineShowFlags.ShaderComplexity);
-		SetupBasePassView(CmdList, View, DrawRenderState, !!ViewFamily.EngineShowFlags.ShaderComplexity);
+		FSceneRenderTargets::Get(CmdList).BeginRenderingGBuffer(CmdList, ERenderTargetLoadAction::ELoad, ERenderTargetLoadAction::ELoad, BasePassDepthStencilAccess, ViewFamily.EngineShowFlags.ShaderComplexity);
+		SetupBasePassView(CmdList, View, DrawRenderState, BasePassDepthStencilAccess, !!ViewFamily.EngineShowFlags.ShaderComplexity);
 	}
 };
 
-void FDeferredShadingSceneRenderer::RenderBasePassViewParallel(FViewInfo& View, FRHICommandListImmediate& ParentCmdList)
+void FDeferredShadingSceneRenderer::RenderBasePassViewParallel(FViewInfo& View, FRHICommandListImmediate& ParentCmdList, FExclusiveDepthStencil::Type BasePassDepthStencilAccess)
 {
 	FBasePassParallelCommandListSet ParallelSet(View, ParentCmdList, 
 		CVarRHICmdBasePassDeferredContexts.GetValueOnRenderThread() > 0, 
 		CVarRHICmdFlushRenderThreadTasksBasePass.GetValueOnRenderThread() == 0 && CVarRHICmdFlushRenderThreadTasks.GetValueOnRenderThread() == 0,
-		ViewFamily);
+		ViewFamily,
+		BasePassDepthStencilAccess);
 
 	RenderBasePassStaticDataParallel(ParallelSet);
 	RenderBasePassDynamicDataParallel(ParallelSet);
 }
 
-void FDeferredShadingSceneRenderer::RenderEditorPrimitives(FRHICommandList& RHICmdList, const FViewInfo& View, bool& bOutDirty) 
+void FDeferredShadingSceneRenderer::RenderEditorPrimitives(FRHICommandList& RHICmdList, const FViewInfo& View, FExclusiveDepthStencil::Type BasePassDepthStencilAccess, bool& bOutDirty) 
 {
 	FDrawingPolicyRenderState DrawRenderState(&RHICmdList, View);
-	SetupBasePassView(RHICmdList, View, DrawRenderState, ViewFamily.EngineShowFlags.ShaderComplexity, true);
+	SetupBasePassView(RHICmdList, View, DrawRenderState, BasePassDepthStencilAccess, ViewFamily.EngineShowFlags.ShaderComplexity, true);
 
 	View.SimpleElementCollector.DrawBatchedElements(RHICmdList, View, FTexture2DRHIRef(), EBlendModeFilter::OpaqueAndMasked);
 
@@ -1221,11 +1241,11 @@ void FDeferredShadingSceneRenderer::RenderEditorPrimitives(FRHICommandList& RHIC
 	}
 }
 
-bool FDeferredShadingSceneRenderer::RenderBasePassView(FRHICommandListImmediate& RHICmdList, FViewInfo& View)
+bool FDeferredShadingSceneRenderer::RenderBasePassView(FRHICommandListImmediate& RHICmdList, FViewInfo& View, FExclusiveDepthStencil::Type BasePassDepthStencilAccess)
 {
 	bool bDirty = false; 
 	FDrawingPolicyRenderState DrawRenderState(&RHICmdList, View);
-	SetupBasePassView(RHICmdList, View, DrawRenderState, ViewFamily.EngineShowFlags.ShaderComplexity);
+	SetupBasePassView(RHICmdList, View, DrawRenderState, BasePassDepthStencilAccess, ViewFamily.EngineShowFlags.ShaderComplexity);
 	bDirty |= RenderBasePassStaticData(RHICmdList, View, DrawRenderState);
 	RenderBasePassDynamicData(RHICmdList, View, DrawRenderState, bDirty);
 
diff --git a/Engine/Source/Runtime/Renderer/Private/BasePassRendering.h b/Engine/Source/Runtime/Renderer/Private/BasePassRendering.h
--- a/Engine/Source/Runtime/Renderer/Private/BasePassRendering.h
+++ b/Engine/Source/Runtime/Renderer/Private/BasePassRendering.h
@@ -92,6 +92,10 @@ public:
 			InstancedNumCulledLightsGrid.SetBuffer(RHICmdList, ShaderRHI, InstancedView.ForwardLightingResources->NumCulledLightsGrid);
 			InstancedCulledLightDataGrid.SetBuffer(RHICmdList, ShaderRHI, InstancedView.ForwardLightingResources->CulledLightDataGrid);
 		}
+		else if (InstancedForwardGlobalLightData.IsBound())
+		{
+			SetUniformBufferParameter(RHICmdList, ShaderRHI, InstancedForwardGlobalLightData, View.ForwardLightingResources->ForwardGlobalLightData);
+		}
 
 		FSceneRenderTargets& SceneRenderTargets = FSceneRenderTargets::Get(RHICmdList);
 
diff --git a/Engine/Source/Runtime/Renderer/Private/DeferredShadingRenderer.cpp b/Engine/Source/Runtime/Renderer/Private/DeferredShadingRenderer.cpp
--- a/Engine/Source/Runtime/Renderer/Private/DeferredShadingRenderer.cpp
+++ b/Engine/Source/Runtime/Renderer/Private/DeferredShadingRenderer.cpp
@@ -111,6 +111,12 @@ static TAutoConsoleVariable<int32> CVarAlphaChannel(
 	TEXT(" 1: enabled"),
 	ECVF_ReadOnly);
 
+static TAutoConsoleVariable<int32> CVarBasePassWriteDepthEvenWithFullPrepass(
+	TEXT("r.BasePassWriteDepthEvenWithFullPrepass"),
+	0,
+	TEXT("0 to allow a readonly base pass, which skips an MSAA depth resolve, and allows masked materials to get EarlyZ (writing to depth while doing clip() disables EarlyZ) (default)\n")
+	TEXT("1 to force depth writes in the base pass.  Useful for debugging when the prepass and base pass don't match what they render."));
+
 DECLARE_CYCLE_STAT(TEXT("PostInitViews FlushDel"), STAT_PostInitViews_FlushDel, STATGROUP_InitViews);
 DECLARE_CYCLE_STAT(TEXT("InitViews Intentional Stall"), STAT_InitViews_Intentional_Stall, STATGROUP_InitViews);
 
@@ -158,6 +164,32 @@ bool ShouldForceFullDepthPass(ERHIFeatureLevel::Type FeatureLevel)
 	return bDBufferAllowed || bStencilLODDither || IsForwardShadingEnabled(FeatureLevel) || UseSelectiveBasePassOutputs();
 }
 
+void GetEarlyZPassMode(ERHIFeatureLevel::Type FeatureLevel, EDepthDrawingMode& EarlyZPassMode, bool& bEarlyZPassMovable)
+{
+	EarlyZPassMode = DDM_NonMaskedOnly;
+	bEarlyZPassMovable = false;
+
+	// developer override, good for profiling, can be useful as project setting
+	{
+		const int32 CVarValue = CVarEarlyZPass.GetValueOnAnyThread();
+
+		switch(CVarValue)
+		{
+			case 0: EarlyZPassMode = DDM_None; break;
+			case 1: EarlyZPassMode = DDM_NonMaskedOnly; break;
+			case 2: EarlyZPassMode = DDM_AllOccluders; break;
+			case 3: break;	// Note: 3 indicates "default behavior" and does not specify an override
+		}
+	}
+
+	if (ShouldForceFullDepthPass(FeatureLevel))
+	{
+		// DBuffer decals and stencil LOD dithering force a full prepass
+		EarlyZPassMode = DDM_AllOpaque;
+		bEarlyZPassMovable = true;
+	}
+}
+
 bool SupportSceneAlpha()
 {
 	return CVarAlphaChannel.GetValueOnRenderThread() != 0;
@@ -192,51 +224,18 @@ const TCHAR* GetDepthPassReason(bool bDitheredLODTransitionsUseStencil, ERHIFeat
 
 FDeferredShadingSceneRenderer::FDeferredShadingSceneRenderer(const FSceneViewFamily* InViewFamily,FHitProxyConsumer* HitProxyConsumer)
 	: FSceneRenderer(InViewFamily, HitProxyConsumer)
-	, EarlyZPassMode(DDM_NonMaskedOnly)
-	, bEarlyZPassMovable(false)
 {
-	if (FPlatformProperties::SupportsWindowedMode())
-	{
-		// Use a depth only pass if we are using full blown HQ lightmaps
-		// Otherwise base pass pixel shaders will be cheap and there will be little benefit to rendering a depth only pass
-		if (AllowHighQualityLightmaps(FeatureLevel) || !ViewFamily.EngineShowFlags.Lighting)
-		{
-			EarlyZPassMode = DDM_None;
-		}
-	}
-
-	// developer override, good for profiling, can be useful as project setting
-	{
-		const int32 CVarValue = CVarEarlyZPass.GetValueOnGameThread();
-
-		switch(CVarValue)
-		{
-			case 0: EarlyZPassMode = DDM_None; break;
-			case 1: EarlyZPassMode = DDM_NonMaskedOnly; break;
-			case 2: EarlyZPassMode = DDM_AllOccluders; break;
-			case 3: break;	// Note: 3 indicates "default behavior" and does not specify an override
-		}
-	}
-
-	// Warning: EarlyZPass logic is mirrored in FStaticMesh::AddToDrawLists
-
-	bEarlyZPassMovable = GEarlyZPassMovable != 0;
-
 	static const auto StencilLODDitherCVar = IConsoleManager::Get().FindTConsoleVariableDataInt(TEXT("r.StencilForLODDither"));
 	bDitheredLODTransitionsUseStencil = StencilLODDitherCVar->GetValueOnAnyThread() != 0;
 
+	GetEarlyZPassMode(FeatureLevel, EarlyZPassMode, bEarlyZPassMovable);
+
 	// Shader complexity requires depth only pass to display masked material cost correctly
 	if (ViewFamily.UseDebugViewPS() && ViewFamily.GetDebugViewShaderMode() != DVSM_OutputMaterialTextureScales)
 	{
 		EarlyZPassMode = DDM_AllOpaque;
 		bEarlyZPassMovable = true;
 	}
-	else if (ShouldForceFullDepthPass(FeatureLevel))
-	{
-		// DBuffer decals and stencil LOD dithering force a full prepass
-		EarlyZPassMode = DDM_AllOpaque;
-		bEarlyZPassMovable = true;
-	}
 }
 
 extern FGlobalBoundShaderState GClearMRTBoundShaderState[8];
@@ -427,7 +426,7 @@ static FORCEINLINE bool NeedsPrePass(const FDeferredShadingSceneRenderer* Render
 		(Renderer->EarlyZPassMode != DDM_None || Renderer->bEarlyZPassMovable != 0);
 }
 
-static void SetAndClearViewGBuffer(FRHICommandListImmediate& RHICmdList, FViewInfo& View, bool bClearDepth)
+static void SetAndClearViewGBuffer(FRHICommandListImmediate& RHICmdList, FViewInfo& View, FExclusiveDepthStencil::Type DepthStencilAccess, bool bClearDepth)
 {
 	// if we didn't to the prepass above, then we will need to clear now, otherwise, it's already been cleared and rendered to
 	ERenderTargetLoadAction DepthLoadAction = bClearDepth ? ERenderTargetLoadAction::EClear : ERenderTargetLoadAction::ELoad;
@@ -437,7 +436,7 @@ static void SetAndClearViewGBuffer(FRHICommandListImmediate& RHICmdList, FViewIn
 	const FLinearColor ClearColor = bClearBlack ? FLinearColor(0, 0, 0, ClearAlpha) : FLinearColor(View.BackgroundColor.R, View.BackgroundColor.G, View.BackgroundColor.B, ClearAlpha);
 
 	// clearing the GBuffer
-	FSceneRenderTargets::Get(RHICmdList).BeginRenderingGBuffer(RHICmdList, ERenderTargetLoadAction::EClear, DepthLoadAction, View.Family->EngineShowFlags.ShaderComplexity, ClearColor);
+	FSceneRenderTargets::Get(RHICmdList).BeginRenderingGBuffer(RHICmdList, ERenderTargetLoadAction::EClear, DepthLoadAction, DepthStencilAccess, View.Family->EngineShowFlags.ShaderComplexity, ClearColor);
 }
 
 static TAutoConsoleVariable<int32> CVarOcclusionQueryLocation(
@@ -781,16 +780,20 @@ void FDeferredShadingSceneRenderer::Render(FRHICommandListImmediate& RHICmdList)
 	const bool bUseVelocityGBuffer = FVelocityRendering::OutputsToGBuffer();
 	const bool bUseSelectiveBasePassOutputs = UseSelectiveBasePassOutputs();
 
-	// Skip the MSAA depth resolve if none of the features which read scene depth before the base pass can be active
-	// This is relying on the post-BasePass ResolveSceneDepthTexture to handle other uses of scene depth, like translucency with DepthFade
-	if ((ViewFamily.EngineShowFlags.Decals && bDBuffer) || ViewFamily.EngineShowFlags.DynamicShadows || bShouldRenderVelocities)
-	{
-		SceneContext.ResolveSceneDepthTexture(RHICmdList, FResolveRect(0, 0, ViewFamily.FamilySizeX, ViewFamily.FamilySizeY));
-	}
-	else
-	{
-		RHICmdList.TransitionResource(EResourceTransitionAccess::EReadable, SceneContext.GetSceneDepthSurface());
-	}
+	// Use readonly depth in the base pass if we have a full depth prepass
+	const bool bAllowReadonlyDepthBasePass = EarlyZPassMode == DDM_AllOpaque
+		&& CVarBasePassWriteDepthEvenWithFullPrepass.GetValueOnRenderThread() == 0
+		&& !ViewFamily.EngineShowFlags.ShaderComplexity
+		&& !ViewFamily.UseDebugViewPS()
+		&& !bIsWireframe
+		&& !ViewFamily.EngineShowFlags.LightMapDensity;
+
+	const FExclusiveDepthStencil::Type BasePassDepthStencilAccess = 
+		bAllowReadonlyDepthBasePass
+		? FExclusiveDepthStencil::DepthRead_StencilWrite 
+		: FExclusiveDepthStencil::DepthWrite_StencilWrite;
+
+	SceneContext.ResolveSceneDepthTexture(RHICmdList, FResolveRect(0, 0, ViewFamily.FamilySizeX, ViewFamily.FamilySizeY));
 
 	if (bUseGBuffer || IsSimpleForwardShadingEnabled(GetFeatureLevelShaderPlatform(FeatureLevel)))
 	{
@@ -860,13 +863,13 @@ void FDeferredShadingSceneRenderer::Render(FRHICommandListImmediate& RHICmdList)
 		ServiceLocalQueue();
 	}
 
-	// Clear the G Buffer render targets
+	// Clear the GBuffer render targets
 	bool bIsGBufferCurrent = false;
 	if (bRequiresRHIClear)
 	{
 		SCOPE_CYCLE_COUNTER(STAT_FDeferredShadingSceneRenderer_SetAndClearViewGBuffer);
 		// set GBuffer to be current, and clear it
-		SetAndClearViewGBuffer(RHICmdList, Views[0], !bDepthWasCleared);
+		SetAndClearViewGBuffer(RHICmdList, Views[0], BasePassDepthStencilAccess, !bDepthWasCleared);
 
 		// depth was cleared now no matter what
 		bDepthWasCleared = true;
@@ -883,7 +886,7 @@ void FDeferredShadingSceneRenderer::Render(FRHICommandListImmediate& RHICmdList)
 	{
 		// make sure the GBuffer is set, in case we didn't need to clear above
 		ERenderTargetLoadAction DepthLoadAction = bDepthWasCleared ? ERenderTargetLoadAction::ELoad : ERenderTargetLoadAction::EClear;
-		SceneContext.BeginRenderingGBuffer(RHICmdList, ERenderTargetLoadAction::ENoAction, DepthLoadAction, ViewFamily.EngineShowFlags.ShaderComplexity);
+		SceneContext.BeginRenderingGBuffer(RHICmdList, ERenderTargetLoadAction::ENoAction, DepthLoadAction, BasePassDepthStencilAccess, ViewFamily.EngineShowFlags.ShaderComplexity);
 	}
 
 	//Single point to catch UE-31578, UE-32536 and UE-22073 and attempt to recover by reallocating Deferred Render Targets
@@ -906,10 +909,15 @@ void FDeferredShadingSceneRenderer::Render(FRHICommandListImmediate& RHICmdList)
 	GRenderTargetPool.AddPhaseEvent(TEXT("BasePass"));
 
 	RHICmdList.SetCurrentStat(GET_STATID(STAT_CLM_BasePass));
-	RenderBasePass(RHICmdList);
+	RenderBasePass(RHICmdList, BasePassDepthStencilAccess);
 	RHICmdList.SetCurrentStat(GET_STATID(STAT_CLM_AfterBasePass));
 	ServiceLocalQueue();
 
+	if (!bAllowReadonlyDepthBasePass)
+	{
+		SceneContext.ResolveSceneDepthTexture(RHICmdList, FResolveRect(0, 0, ViewFamily.FamilySizeX, ViewFamily.FamilySizeY));
+	}
+
 	if (ViewFamily.EngineShowFlags.VisualizeLightCulling)
 	{
 		// clear out emissive and baked lighting (not too efficient but simple and only needed for this debug view)
@@ -937,7 +945,6 @@ void FDeferredShadingSceneRenderer::Render(FRHICommandListImmediate& RHICmdList)
 		bRequiresFarZQuadClear = false;
 	}
 	
-	SceneContext.ResolveSceneDepthTexture(RHICmdList, FResolveRect(0, 0, ViewFamily.FamilySizeX, ViewFamily.FamilySizeY));
 	SceneContext.ResolveSceneDepthToAuxiliaryTexture(RHICmdList);
 
 	bool bOcclusionAfterBasePass = bIsOcclusionTesting && !bOcclusionBeforeBasePass;
@@ -1096,7 +1103,7 @@ void FDeferredShadingSceneRenderer::Render(FRHICommandListImmediate& RHICmdList)
 		ServiceLocalQueue();
 
 		// SSS need the SceneColor finalized as an SRV.
-		SceneContext.FinishRenderingSceneColor(RHICmdList);
+		ResolveSceneColor(RHICmdList);
 
 		// Render reflections that only operate on opaque pixels
 		RenderDeferredReflections(RHICmdList, DynamicBentNormalAO, VelocityRT);
@@ -1238,7 +1245,7 @@ void FDeferredShadingSceneRenderer::Render(FRHICommandListImmediate& RHICmdList)
 	}
 
 	// Resolve the scene color for post processing.
-	SceneContext.ResolveSceneColor(RHICmdList, FResolveRect(0, 0, ViewFamily.FamilySizeX, ViewFamily.FamilySizeY));
+	ResolveSceneColor(RHICmdList);
 
 	GetRendererModule().RenderPostResolvedSceneColorExtension(RHICmdList, SceneContext);
 
diff --git a/Engine/Source/Runtime/Renderer/Private/DeferredShadingRenderer.h b/Engine/Source/Runtime/Renderer/Private/DeferredShadingRenderer.h
--- a/Engine/Source/Runtime/Renderer/Private/DeferredShadingRenderer.h
+++ b/Engine/Source/Runtime/Renderer/Private/DeferredShadingRenderer.h
@@ -97,19 +97,19 @@ public:
 	void RenderBasePassDynamicDataParallel(FParallelCommandListSet& ParallelCommandListSet);
 
 	/** Renders the basepass for a given View, in parallel */
-	void RenderBasePassViewParallel(FViewInfo& View, FRHICommandListImmediate& ParentCmdList);
+	void RenderBasePassViewParallel(FViewInfo& View, FRHICommandListImmediate& ParentCmdList, FExclusiveDepthStencil::Type BasePassDepthStencilAccess);
 
 	/** Renders the basepass for a given View. */
-	bool RenderBasePassView(FRHICommandListImmediate& RHICmdList, FViewInfo& View);
+	bool RenderBasePassView(FRHICommandListImmediate& RHICmdList, FViewInfo& View, FExclusiveDepthStencil::Type BasePassDepthStencilAccess);
 
 	/** Renders editor primitives for a given View. */
-	void RenderEditorPrimitives(FRHICommandList& RHICmdList, const FViewInfo& View, bool& bOutDirty);
+	void RenderEditorPrimitives(FRHICommandList& RHICmdList, const FViewInfo& View, FExclusiveDepthStencil::Type BasePassDepthStencilAccess, bool& bOutDirty);
 	
 	/** 
 	* Renders the scene's base pass 
 	* @return true if anything was rendered
 	*/
-	bool RenderBasePass(FRHICommandListImmediate& RHICmdList);
+	bool RenderBasePass(FRHICommandListImmediate& RHICmdList, FExclusiveDepthStencil::Type BasePassDepthStencilAccess);
 
 	/** Finishes the view family rendering. */
 	void RenderFinish(FRHICommandListImmediate& RHICmdList);
@@ -173,6 +173,8 @@ private:
 	*/
 	bool PreRenderPrePass(FRHICommandListImmediate& RHICmdList);
 
+	void RenderPrePassEditorPrimitives(FRHICommandList& RHICmdList, const FViewInfo& View, FDepthDrawingPolicyFactory::ContextType Context);
+
 	/**
 	 * Renders the scene's prepass and occlusion queries.
 	 * @return true if the depth was cleared
diff --git a/Engine/Source/Runtime/Renderer/Private/DepthRendering.cpp b/Engine/Source/Runtime/Renderer/Private/DepthRendering.cpp
--- a/Engine/Source/Runtime/Renderer/Private/DepthRendering.cpp
+++ b/Engine/Source/Runtime/Renderer/Private/DepthRendering.cpp
@@ -24,6 +24,7 @@
 #include "IHeadMountedDisplay.h"
 #include "ScreenRendering.h"
 #include "PostProcess/SceneFilterRendering.h"
+#include "DynamicPrimitiveDrawing.h"
 
 static TAutoConsoleVariable<int32> CVarRHICmdPrePassDeferredContexts(
 	TEXT("r.RHICmdPrePassDeferredContexts"),
@@ -868,6 +869,8 @@ bool FDeferredShadingSceneRenderer::RenderPrePassViewDynamic(FRHICommandList& RH
 		}
 	}
 
+	RenderPrePassEditorPrimitives(RHICmdList, View, Context);
+
 	return true;
 }
 
@@ -1183,6 +1186,33 @@ bool FDeferredShadingSceneRenderer::PreRenderPrePass(FRHICommandListImmediate& R
 	return bDepthWasCleared;
 }
 
+void FDeferredShadingSceneRenderer::RenderPrePassEditorPrimitives(FRHICommandList& RHICmdList, const FViewInfo& View, FDepthDrawingPolicyFactory::ContextType Context) 
+{
+	FDrawingPolicyRenderState DrawRenderState(&RHICmdList, View);
+	SetupPrePassView(RHICmdList, View, DrawRenderState);
+
+	View.SimpleElementCollector.DrawBatchedElements(RHICmdList, View, FTexture2DRHIRef(), EBlendModeFilter::OpaqueAndMasked);
+
+	bool bDirty = false;
+	if (!View.Family->EngineShowFlags.CompositeEditorPrimitives)
+	{
+		const auto ShaderPlatform = View.GetShaderPlatform();
+		const bool bNeedToSwitchVerticalAxis = RHINeedsToSwitchVerticalAxis(ShaderPlatform);
+
+		// Draw the base pass for the view's batched mesh elements.
+		bDirty |= DrawViewElements<FDepthDrawingPolicyFactory>(RHICmdList, View, DrawRenderState, Context, SDPG_World, true) || bDirty;
+
+		// Draw the view's batched simple elements(lines, sprites, etc).
+		bDirty |= View.BatchedViewElements.Draw(RHICmdList, FeatureLevel, bNeedToSwitchVerticalAxis, View, false) || bDirty;
+
+		// Draw foreground objects last
+		bDirty |= DrawViewElements<FDepthDrawingPolicyFactory>(RHICmdList, View, DrawRenderState, Context, SDPG_Foreground, true) || bDirty;
+
+		// Draw the view's batched simple elements(lines, sprites, etc).
+		bDirty |= View.TopBatchedViewElements.Draw(RHICmdList, FeatureLevel, bNeedToSwitchVerticalAxis, View, false) || bDirty;
+	}
+}
+
 bool FDeferredShadingSceneRenderer::RenderPrePass(FRHICommandListImmediate& RHICmdList, TFunctionRef<void()> AfterTasksAreStarted)
 {
 	bool bDepthWasCleared = false;
diff --git a/Engine/Source/Runtime/Renderer/Private/DepthRendering.h b/Engine/Source/Runtime/Renderer/Private/DepthRendering.h
--- a/Engine/Source/Runtime/Renderer/Private/DepthRendering.h
+++ b/Engine/Source/Runtime/Renderer/Private/DepthRendering.h
@@ -21,11 +21,11 @@ enum EDepthDrawingMode
 {
 	// tested at a higher level
 	DDM_None			= 0,
-	//
+	// Opaque materials only
 	DDM_NonMaskedOnly	= 1,
-	//
+	// Opaque and masked materials, but no objects with bUseAsOccluder disabled
 	DDM_AllOccluders	= 2,
-	//
+	// Full prepass, every object must be drawn and every pixel must match the base pass depth
 	DDM_AllOpaque		= 3,
 };
 
diff --git a/Engine/Source/Runtime/Renderer/Private/DistortionRendering.cpp b/Engine/Source/Runtime/Renderer/Private/DistortionRendering.cpp
--- a/Engine/Source/Runtime/Renderer/Private/DistortionRendering.cpp
+++ b/Engine/Source/Runtime/Renderer/Private/DistortionRendering.cpp
@@ -1186,8 +1186,8 @@ void FSceneRenderer::RenderDistortionES2(FRHICommandListImmediate& RHICmdList)
 		
 		FSceneRenderTargets& SceneContext = FSceneRenderTargets::Get(RHICmdList);
 
-		SceneContext.ResolveSceneColor(RHICmdList, FResolveRect(0, 0, ViewFamily.FamilySizeX, ViewFamily.FamilySizeY));
-	
+		RHICmdList.CopyToResolveTarget(SceneContext.GetSceneColorSurface(), SceneContext.GetSceneColorTexture(), true, FResolveRect(0, 0, ViewFamily.FamilySizeX, ViewFamily.FamilySizeY));
+
 		TRefCountPtr<IPooledRenderTarget> SceneColorDistorted;
 		GRenderTargetPool.FindFreeElement(RHICmdList, SceneContext.GetSceneColor()->GetDesc(), SceneColorDistorted, TEXT("SceneColorDistorted"));
 		const FSceneRenderTargetItem& DistortedRenderTarget = SceneColorDistorted->GetRenderTargetItem();
diff --git a/Engine/Source/Runtime/Renderer/Private/MobileShadingRenderer.cpp b/Engine/Source/Runtime/Renderer/Private/MobileShadingRenderer.cpp
--- a/Engine/Source/Runtime/Renderer/Private/MobileShadingRenderer.cpp
+++ b/Engine/Source/Runtime/Renderer/Private/MobileShadingRenderer.cpp
@@ -267,7 +267,7 @@ void FMobileSceneRenderer::Render(FRHICommandListImmediate& RHICmdList)
 	if (!bGammaSpace || bRenderToSceneColor)
 	{
 		// Resolve the scene color for post processing.
-		SceneContext.ResolveSceneColor(RHICmdList, FResolveRect(0, 0, ViewFamily.FamilySizeX, ViewFamily.FamilySizeY));
+		RHICmdList.CopyToResolveTarget(SceneContext.GetSceneColorSurface(), SceneContext.GetSceneColorTexture(), true, FResolveRect(0, 0, ViewFamily.FamilySizeX, ViewFamily.FamilySizeY));
 
 		// On PowerVR we see flickering of shadows and depths not updating correctly if targets are discarded.
 		// See CVarMobileForceDepthResolve use in ConditionalResolveSceneDepth.
diff --git a/Engine/Source/Runtime/Renderer/Private/MobileTranslucentRendering.cpp b/Engine/Source/Runtime/Renderer/Private/MobileTranslucentRendering.cpp
--- a/Engine/Source/Runtime/Renderer/Private/MobileTranslucentRendering.cpp
+++ b/Engine/Source/Runtime/Renderer/Private/MobileTranslucentRendering.cpp
@@ -73,7 +73,7 @@ void FMobileSceneRenderer::CopySceneAlpha(FRHICommandListImmediate& RHICmdList,
 	RHICmdList.SetDepthStencilState(TStaticDepthStencilState<false,CF_Always>::GetRHI());
 	RHICmdList.SetBlendState(TStaticBlendState<>::GetRHI());
 
-	SceneContext.ResolveSceneColor(RHICmdList);
+	RHICmdList.CopyToResolveTarget(SceneContext.GetSceneColorSurface(), SceneContext.GetSceneColorTexture(), true, FResolveRect(0, 0, ViewFamily.FamilySizeX, ViewFamily.FamilySizeY));
 
 	SceneContext.BeginRenderingSceneAlphaCopy(RHICmdList);
 
diff --git a/Engine/Source/Runtime/Renderer/Private/PostProcess/PostProcessCompositeEditorPrimitives.cpp b/Engine/Source/Runtime/Renderer/Private/PostProcess/PostProcessCompositeEditorPrimitives.cpp
--- a/Engine/Source/Runtime/Renderer/Private/PostProcess/PostProcessCompositeEditorPrimitives.cpp
+++ b/Engine/Source/Runtime/Renderer/Private/PostProcess/PostProcessCompositeEditorPrimitives.cpp
@@ -206,6 +206,7 @@ void FRCPassPostProcessCompositeEditorPrimitives::Process(FRenderingCompositePas
 	const FSceneViewFamily& ViewFamily = *(View.Family);
 	
 	FDrawingPolicyRenderState DrawRenderState(&Context.RHICmdList, View);
+	DrawRenderState.SetDepthStencilAccess(FExclusiveDepthStencil::DepthWrite_StencilWrite);
 
 	FIntRect SrcRect = View.ViewRect;
 	FIntRect DestRect = View.ViewRect;
diff --git a/Engine/Source/Runtime/Renderer/Private/PostProcess/PostProcessing.cpp b/Engine/Source/Runtime/Renderer/Private/PostProcess/PostProcessing.cpp
--- a/Engine/Source/Runtime/Renderer/Private/PostProcess/PostProcessing.cpp
+++ b/Engine/Source/Runtime/Renderer/Private/PostProcess/PostProcessing.cpp
@@ -1267,7 +1267,7 @@ void FPostProcessing::Process(FRHICommandListImmediate& RHICmdList, const FViewI
 		}
 
 		static const auto CVarHDROutputEnabled = IConsoleManager::Get().FindTConsoleVariableDataInt(TEXT("r.HDR.EnableHDROutput"));
-		const bool bHDROutputEnabled = GRHISupportsHDROutput && CVarHDROutputEnabled->GetValueOnRenderThread() != 0;
+		const bool bHDROutputEnabled = GRHISupportsHDROutput && CVarHDROutputEnabled && CVarHDROutputEnabled->GetValueOnRenderThread() != 0;
 
 		static const auto CVarDumpFramesAsHDR = IConsoleManager::Get().FindTConsoleVariableDataInt(TEXT("r.BufferVisualizationDumpFramesAsHDR"));
 		const bool bHDRTonemapperOutput = bAllowTonemapper && (GetHighResScreenshotConfig().bCaptureHDR || CVarDumpFramesAsHDR->GetValueOnRenderThread() || bHDROutputEnabled);
diff --git a/Engine/Source/Runtime/Renderer/Private/PostProcess/SceneRenderTargets.cpp b/Engine/Source/Runtime/Renderer/Private/PostProcess/SceneRenderTargets.cpp
--- a/Engine/Source/Runtime/Renderer/Private/PostProcess/SceneRenderTargets.cpp
+++ b/Engine/Source/Runtime/Renderer/Private/PostProcess/SceneRenderTargets.cpp
@@ -13,8 +13,6 @@
 #include "RendererModule.h"
 #include "LightPropagationVolume.h"
 #include "ScenePrivate.h"
-#include "HdrCustomResolveShaders.h"
-#include "WideCustomResolveShaders.h"
 #include "ClearQuad.h"
 #include "RenderUtils.h"
 
@@ -70,17 +68,6 @@ static TAutoConsoleVariable<int32> CVarMSAACount(
 	ECVF_RenderThreadSafe | ECVF_Scalability
 	);
 
-static TAutoConsoleVariable<int32> CVarWideCustomResolve(
-	TEXT("r.WideCustomResolve"),
-	0,
-	TEXT("Use a wide custom resolve filter when MSAA is enabled")
-	TEXT("0: Disabled [hardware box filter]")
-	TEXT("1: Wide (r=1.25, 12 samples)")
-	TEXT("2: Wider (r=1.4, 16 samples)")
-	TEXT("3: Widest (r=1.5, 20 samples)"),
-	ECVF_RenderThreadSafe | ECVF_Scalability
-	);
-
 static TAutoConsoleVariable<int32> CVarMobileMSAA(
 	TEXT("r.MobileMSAA"),
 	1,
@@ -561,21 +548,49 @@ int32 FSceneRenderTargets::GetGBufferRenderTargets(ERenderTargetLoadAction Color
 	return MRTCount;
 }
 
-void FSceneRenderTargets::BeginRenderingGBuffer(FRHICommandList& RHICmdList, ERenderTargetLoadAction ColorLoadAction, ERenderTargetLoadAction DepthLoadAction, bool bBindQuadOverdrawBuffers, const FLinearColor& ClearColor/*=(0,0,0,1)*/)
+void FSceneRenderTargets::SetQuadOverdrawUAV(FRHICommandList& RHICmdList, bool bBindQuadOverdrawBuffers, FRHISetRenderTargetsInfo& Info)
+{
+	if (bBindQuadOverdrawBuffers && AllowDebugViewPS(DVSM_QuadComplexity, GetFeatureLevelShaderPlatform(CurrentFeatureLevel)))
+	{
+		if (QuadOverdrawBuffer.IsValid() && QuadOverdrawBuffer->GetRenderTargetItem().UAV.IsValid())
+		{
+			QuadOverdrawIndex = IsAnyForwardShadingEnabled(GetFeatureLevelShaderPlatform(CurrentFeatureLevel)) ? 1 : 7; // As defined in QuadOverdraw.usf
+
+			// Increase the rendertarget count in order to control the bound slot of the UAV.
+			check(Info.NumColorRenderTargets <= QuadOverdrawIndex);
+			Info.NumColorRenderTargets = QuadOverdrawIndex;
+			Info.UnorderedAccessView[Info.NumUAVs++] = QuadOverdrawBuffer->GetRenderTargetItem().UAV;
+
+			// Clear to default value
+			const uint32 ClearValue[4] = { 0, 0, 0, 0 };
+			RHICmdList.ClearUAV(QuadOverdrawBuffer->GetRenderTargetItem().UAV, ClearValue);
+			RHICmdList.TransitionResource(EResourceTransitionAccess::ERWBarrier, EResourceTransitionPipeline::EGfxToGfx, QuadOverdrawBuffer->GetRenderTargetItem().UAV);
+		}
+	}
+}
+
+void FSceneRenderTargets::BeginRenderingGBuffer(FRHICommandList& RHICmdList, ERenderTargetLoadAction ColorLoadAction, ERenderTargetLoadAction DepthLoadAction, FExclusiveDepthStencil::Type DepthStencilAccess, bool bBindQuadOverdrawBuffers, const FLinearColor& ClearColor/*=(0,0,0,1)*/)
 {
 	SCOPED_DRAW_EVENT(RHICmdList, BeginRenderingSceneColor);
+	check(CurrentFeatureLevel >= ERHIFeatureLevel::SM4);
+	AllocSceneColor(RHICmdList);
+
+	FRHIRenderTargetView RenderTargets[MaxSimultaneousRenderTargets] = { NULL };
+	const ERenderTargetStoreAction DepthStoreAction = (DepthStencilAccess & FExclusiveDepthStencil::DepthWrite) ? ERenderTargetStoreAction::EStore : ERenderTargetStoreAction::ENoAction;
+	FRHIDepthRenderTargetView DepthView(GetSceneDepthSurface(), DepthLoadAction, DepthStoreAction, DepthStencilAccess);
 
 	if (IsAnyForwardShadingEnabled(GetFeatureLevelShaderPlatform(CurrentFeatureLevel)))
 	{
-		// in this non-standard case, just render to scene color with default mode
-		BeginRenderingSceneColor(RHICmdList, ColorLoadAction == ERenderTargetLoadAction::EClear ? ESimpleRenderTargetMode::EClearColorExistingDepth : ESimpleRenderTargetMode::EUninitializedColorExistingDepth);
-		return;
-	}
+		int32 MRTCount = 0;
+		RenderTargets[MRTCount++] = FRHIRenderTargetView(GetSceneColorSurface(), 0, -1, ColorLoadAction, ERenderTargetStoreAction::EStore);
 
-	AllocSceneColor(RHICmdList);
+		FRHISetRenderTargetsInfo Info(MRTCount, RenderTargets, DepthView);	
 
-	// Set the scene color surface as the render target, and the scene depth surface as the depth-stencil target.
-	if (CurrentFeatureLevel >= ERHIFeatureLevel::SM4)
+		SetQuadOverdrawUAV(RHICmdList, bBindQuadOverdrawBuffers, Info);
+
+		RHICmdList.SetRenderTargetsAndClear(Info);
+	}
+	else
 	{
 		bool bClearColor = ColorLoadAction == ERenderTargetLoadAction::EClear;
 		bool bClearDepth = DepthLoadAction == ERenderTargetLoadAction::EClear;
@@ -585,10 +600,10 @@ void FSceneRenderTargets::BeginRenderingGBuffer(FRHICommandList& RHICmdList, ERe
 		const FTextureRHIRef& SceneColorTex = GetSceneColorSurface();
 		bool bShaderClear = false;
 		if (bClearColor)
-		{			
+		{
 			if (!SceneColorTex->HasClearValue() || (ClearColor != SceneColorTex->GetClearColor()))
 			{
-				ColorLoadAction = ERenderTargetLoadAction::ENoAction;				
+				ColorLoadAction = ERenderTargetLoadAction::ENoAction;
 				bShaderClear = true;
 			}
 			else
@@ -598,38 +613,20 @@ void FSceneRenderTargets::BeginRenderingGBuffer(FRHICommandList& RHICmdList, ERe
 		}
 
 		int32 VelocityRTIndex;
-		FRHIRenderTargetView RenderTargets[MaxSimultaneousRenderTargets];
 		int32 MRTCount = GetGBufferRenderTargets(ColorLoadAction, RenderTargets, VelocityRTIndex);
 
 		//make sure our conditions for shader clear fallback are valid.
 		check(RenderTargets[0].Texture == SceneColorTex);
 
-		FRHIDepthRenderTargetView DepthView(GetSceneDepthSurface(), DepthLoadAction, ERenderTargetStoreAction::EStore);
-		FRHISetRenderTargetsInfo Info(MRTCount, RenderTargets, DepthView);		
+		FRHISetRenderTargetsInfo Info(MRTCount, RenderTargets, DepthView);
 
 		if (bClearDepth)
 		{
-			bSceneDepthCleared = true;			
-		}		
-
-		if (bBindQuadOverdrawBuffers && AllowDebugViewPS(DVSM_QuadComplexity, GetFeatureLevelShaderPlatform(CurrentFeatureLevel)))
-		{
-			if (QuadOverdrawBuffer.IsValid() && QuadOverdrawBuffer->GetRenderTargetItem().UAV.IsValid())
-			{
-				QuadOverdrawIndex = 7; // As defined in QuadOverdraw.usf
-
-				// Increase the rendertarget count in order to control the bound slot of the UAV.
-				check(Info.NumColorRenderTargets <= QuadOverdrawIndex);
-				Info.NumColorRenderTargets = QuadOverdrawIndex;
-				Info.UnorderedAccessView[Info.NumUAVs++] = QuadOverdrawBuffer->GetRenderTargetItem().UAV;
-
-				// Clear to default value
-				const uint32 ClearValue[4] = { 0, 0, 0, 0 };
-				RHICmdList.ClearUAV(QuadOverdrawBuffer->GetRenderTargetItem().UAV, ClearValue);
-				RHICmdList.TransitionResource(EResourceTransitionAccess::ERWBarrier, EResourceTransitionPipeline::EGfxToGfx, QuadOverdrawBuffer->GetRenderTargetItem().UAV);
-			}
+			bSceneDepthCleared = true;
 		}
 
+		SetQuadOverdrawUAV(RHICmdList, bBindQuadOverdrawBuffers, Info);
+
 		// set the render target
 		RHICmdList.SetRenderTargetsAndClear(Info);
 		if (bShaderClear)
@@ -1097,13 +1094,6 @@ void FSceneRenderTargets::AdjustGBufferRefCount(FRHICommandList& RHICmdList, int
 	}	
 }
 
-void FSceneRenderTargets::FinishRenderingSceneColor(FRHICommandListImmediate& RHICmdList, const FResolveRect& ResolveRect)
-{
-	SCOPED_DRAW_EVENT(RHICmdList, FinishRenderingSceneColor);
-
-	ResolveSceneColor(RHICmdList);
-}
-
 bool FSceneRenderTargets::BeginRenderingCustomDepth(FRHICommandListImmediate& RHICmdList, bool bPrimitives)
 {
 	IPooledRenderTarget* CustomDepthRenderTarget = RequestCustomDepth(RHICmdList, bPrimitives);
@@ -1150,90 +1140,6 @@ void FSceneRenderTargets::FinishRenderingCustomDepth(FRHICommandListImmediate& R
 	bCustomDepthIsValid = true;
 }
 
-/**
-* Saves a previously rendered scene color target
-*/
-void FSceneRenderTargets::ResolveSceneColor(FRHICommandList& RHICmdList, const FResolveRect& ResolveRect)
-{
-	SCOPED_DRAW_EVENT(RHICmdList, ResolveSceneColor);
-
-	auto& CurrentSceneColor = GetSceneColor();
-	uint32 CurrentNumSamples = CurrentSceneColor->GetDesc().NumSamples;
-
-	const EShaderPlatform CurrentShaderPlatform = GShaderPlatformForFeatureLevel[CurrentFeatureLevel];
-	if (CurrentNumSamples <= 1 || !RHISupportsSeparateMSAAAndResolveTextures(CurrentShaderPlatform))
-	{
-		RHICmdList.CopyToResolveTarget(GetSceneColorSurface(), GetSceneColorTexture(), true, FResolveParams(ResolveRect));
-	}
-	else 
-	{
-		// Custom shader based color resolve for HDR color to emulate mobile.
-		SetRenderTarget(RHICmdList, GetSceneColorTexture(), FTextureRHIParamRef());
-		
-		if (ResolveRect.IsValid())
-		{
-			RHICmdList.SetScissorRect(true, ResolveRect.X1, ResolveRect.Y1, ResolveRect.X2, ResolveRect.Y2);
-		}
-
-		RHICmdList.SetBlendState(TStaticBlendState<>::GetRHI());
-		RHICmdList.SetRasterizerState(TStaticRasterizerState<>::GetRHI());
-		RHICmdList.SetDepthStencilState(TStaticDepthStencilState<false, CF_Always>::GetRHI());
-		RHICmdList.SetStreamSource(0, NULL, 0, 0);
-
-		int32 ResolveWidth = CVarWideCustomResolve.GetValueOnRenderThread();
-
-		if (CurrentNumSamples <= 1)
-		{
-			ResolveWidth = 0;
-		}
-
-		if (ResolveWidth != 0)
-		{
-			ResolveFilterWide(RHICmdList, CurrentFeatureLevel, CurrentSceneColor->GetRenderTargetItem().TargetableTexture, FIntPoint(0, 0), CurrentNumSamples, ResolveWidth);
-		}
-		else
-		{
-			auto ShaderMap = GetGlobalShaderMap(CurrentFeatureLevel);
-			TShaderMapRef<FHdrCustomResolveVS> VertexShader(ShaderMap);
-
-			if (CurrentNumSamples == 2)
-			{
-				TShaderMapRef<FHdrCustomResolve2xPS> PixelShader(ShaderMap);
-				static FGlobalBoundShaderState BoundShaderState;
-				SetGlobalBoundShaderState(RHICmdList, CurrentFeatureLevel, BoundShaderState, GetVertexDeclarationFVector4(), *VertexShader, *PixelShader);
-				PixelShader->SetParameters(RHICmdList, CurrentSceneColor->GetRenderTargetItem().TargetableTexture);
-				RHICmdList.DrawPrimitive(PT_TriangleList, 0, 1, 1);
-			}
-			else if (CurrentNumSamples == 4)
-			{
-				TShaderMapRef<FHdrCustomResolve4xPS> PixelShader(ShaderMap);
-				static FGlobalBoundShaderState BoundShaderState;
-				SetGlobalBoundShaderState(RHICmdList, CurrentFeatureLevel, BoundShaderState, GetVertexDeclarationFVector4(), *VertexShader, *PixelShader);
-				PixelShader->SetParameters(RHICmdList, CurrentSceneColor->GetRenderTargetItem().TargetableTexture);
-				RHICmdList.DrawPrimitive(PT_TriangleList, 0, 1, 1);
-			}
-			else if (CurrentNumSamples == 8)
-			{
-				TShaderMapRef<FHdrCustomResolve8xPS> PixelShader(ShaderMap);
-				static FGlobalBoundShaderState BoundShaderState;
-				SetGlobalBoundShaderState(RHICmdList, CurrentFeatureLevel, BoundShaderState, GetVertexDeclarationFVector4(), *VertexShader, *PixelShader);
-				PixelShader->SetParameters(RHICmdList, CurrentSceneColor->GetRenderTargetItem().TargetableTexture);
-				RHICmdList.DrawPrimitive(PT_TriangleList, 0, 1, 1);
-			}
-			else
-			{
-				// Everything other than 2,4,8 samples is not implemented.
-				check(0);
-			}
-		}
-
-		if (ResolveRect.IsValid())
-		{
-			RHICmdList.SetScissorRect(false, 0, 0, 0, 0);
-		}
-	}
-}
-
 void FSceneRenderTargets::BeginRenderingPrePass(FRHICommandList& RHICmdList, bool bPerformClear)
 {
 	SCOPED_DRAW_EVENT(RHICmdList, BeginRenderingPrePass);
@@ -1353,12 +1259,6 @@ void FSceneRenderTargets::BeginRenderingTranslucency(FRHICommandList& RHICmdList
 	RHICmdList.SetViewport(View.ViewRect.Min.X, View.ViewRect.Min.Y, 0.0f, View.ViewRect.Max.X, View.ViewRect.Max.Y, 1.0f);
 }
 
-void FSceneRenderTargets::FinishRenderingTranslucency(FRHICommandListImmediate& RHICmdList, const class FViewInfo& View)
-{
-	FinishRenderingSceneColor(RHICmdList);
-}
-
-
 void FSceneRenderTargets::BeginRenderingSeparateTranslucency(FRHICommandList& RHICmdList, const FViewInfo& View, bool bFirstTimeThisFrame)
 {
 	check(IsSeparateTranslucencyActive(View));
diff --git a/Engine/Source/Runtime/Renderer/Private/PostProcess/SceneRenderTargets.h b/Engine/Source/Runtime/Renderer/Private/PostProcess/SceneRenderTargets.h
--- a/Engine/Source/Runtime/Renderer/Private/PostProcess/SceneRenderTargets.h
+++ b/Engine/Source/Runtime/Renderer/Private/PostProcess/SceneRenderTargets.h
@@ -227,7 +227,8 @@ public:
 
 	void SetSeparateTranslucencyBufferSize(bool bAnyViewWantsDownsampledSeparateTranslucency);
 
-	void BeginRenderingGBuffer(FRHICommandList& RHICmdList, ERenderTargetLoadAction ColorLoadAction, ERenderTargetLoadAction DepthLoadAction, bool bBindQuadOverdrawBuffers, const FLinearColor& ClearColor = FLinearColor(0, 0, 0, 1));
+	void SetQuadOverdrawUAV(FRHICommandList& RHICmdList, bool bBindQuadOverdrawBuffers, FRHISetRenderTargetsInfo& Info);
+	void BeginRenderingGBuffer(FRHICommandList& RHICmdList, ERenderTargetLoadAction ColorLoadAction, ERenderTargetLoadAction DepthLoadAction, FExclusiveDepthStencil::Type DepthStencilAccess, bool bBindQuadOverdrawBuffers, const FLinearColor& ClearColor = FLinearColor(0, 0, 0, 1));
 	void FinishRenderingGBuffer(FRHICommandListImmediate& RHICmdList);
 
 	/**
@@ -235,11 +236,6 @@ public:
 	 */
 	void BeginRenderingSceneColor(FRHICommandList& FRHICommandListImmediate, ESimpleRenderTargetMode RenderTargetMode = ESimpleRenderTargetMode::EUninitializedColorExistingDepth, FExclusiveDepthStencil DepthStencilAccess = FExclusiveDepthStencil::DepthWrite_StencilWrite, bool bTransitionWritable = true);
 
-	/**
-	 * Called when finished rendering to the scene color surface
-	 */
-	void FinishRenderingSceneColor(FRHICommandListImmediate& RHICmdList, const FResolveRect& ResolveRect = FResolveRect());
-
 	void BeginRenderingSceneMonoColor(FRHICommandList& RHICmdList, ESimpleRenderTargetMode RenderTargetMode = ESimpleRenderTargetMode::EUninitializedColorExistingDepth, FExclusiveDepthStencil DepthStencilAccess = FExclusiveDepthStencil::DepthWrite_StencilWrite);
 
 	// @return true: call FinishRenderingCustomDepth after rendering, false: don't render it, feature is disabled
@@ -247,16 +243,10 @@ public:
 	// only call if BeginRenderingCustomDepth() returned true
 	void FinishRenderingCustomDepth(FRHICommandListImmediate& RHICmdList, const FResolveRect& ResolveRect = FResolveRect());
 
-	/**
-	 * Resolve a previously rendered scene color surface.
-	 */
-	void ResolveSceneColor(FRHICommandList& RHICmdList, const FResolveRect& ResolveRect = FResolveRect());
-
 	/** Binds the appropriate shadow depth cube map for rendering. */
 	void BeginRenderingCubeShadowDepth(FRHICommandList& RHICmdList, int32 ShadowResolution);
 
 	void BeginRenderingTranslucency(FRHICommandList& RHICmdList, const class FViewInfo& View, bool bFirstTimeThisFrame = true);
-	void FinishRenderingTranslucency(FRHICommandListImmediate& RHICmdList, const class FViewInfo& View);
 
 	void BeginRenderingSeparateTranslucency(FRHICommandList& RHICmdList, const FViewInfo& View, bool bFirstTimeThisFrame);
 	void FinishRenderingSeparateTranslucency(FRHICommandList& RHICmdList);
diff --git a/Engine/Source/Runtime/Renderer/Private/ReflectionEnvironment.cpp b/Engine/Source/Runtime/Renderer/Private/ReflectionEnvironment.cpp
--- a/Engine/Source/Runtime/Renderer/Private/ReflectionEnvironment.cpp
+++ b/Engine/Source/Runtime/Renderer/Private/ReflectionEnvironment.cpp
@@ -894,7 +894,7 @@ void FDeferredShadingSceneRenderer::RenderReflectionCaptureSpecularBounceForAllV
 			EDRF_UseTriangleOptimization);
 	}
 
-	SceneContext.FinishRenderingSceneColor(RHICmdList);
+	ResolveSceneColor(RHICmdList);
 }
 
 bool FDeferredShadingSceneRenderer::ShouldDoReflectionEnvironment() const
@@ -1108,7 +1108,7 @@ void FDeferredShadingSceneRenderer::RenderTiledDeferredImageBasedReflections(FRH
 
 	TRefCountPtr<IPooledRenderTarget> NewSceneColor;
 	{
-		SceneContext.ResolveSceneColor(RHICmdList, FResolveRect(0, 0, ViewFamily.FamilySizeX, ViewFamily.FamilySizeY));
+		ResolveSceneColor(RHICmdList);
 
 		FPooledRenderTargetDesc Desc = SceneContext.GetSceneColor()->GetDesc();
 		Desc.TargetableFlags |= TexCreate_UAV;
@@ -1449,7 +1449,7 @@ void FDeferredShadingSceneRenderer::RenderStandardDeferredImageBasedReflections(
 				SceneContext.GetBufferSizeXY(),
 				*VertexShader);
 
-			SceneContext.FinishRenderingSceneColor(RHICmdList);
+			ResolveSceneColor(RHICmdList);
 		}
 	}
 }
diff --git a/Engine/Source/Runtime/Renderer/Private/SceneCaptureRendering.cpp b/Engine/Source/Runtime/Renderer/Private/SceneCaptureRendering.cpp
--- a/Engine/Source/Runtime/Renderer/Private/SceneCaptureRendering.cpp
+++ b/Engine/Source/Runtime/Renderer/Private/SceneCaptureRendering.cpp
@@ -397,6 +397,7 @@ FSceneRenderer* CreateSceneRendererForSceneCapture(
 		ViewInitOptions.StereoPass = ViewState.StereoPass;
 		ViewInitOptions.SceneViewStateInterface = (ViewInitOptions.StereoPass != EStereoscopicPass::eSSP_RIGHT_EYE) ? SceneCaptureComponent->GetViewState() : SceneCaptureComponent->GetStereoViewState();
 		ViewInitOptions.ProjectionMatrix = ViewState.ProjectionMatrix;
+		ViewInitOptions.LODDistanceFactor = FMath::Clamp(SceneCaptureComponent->LODDistanceFactor, .01f, 100.0f);
 
 		if (bCaptureSceneColor)
 		{
diff --git a/Engine/Source/Runtime/Renderer/Private/SceneCore.cpp b/Engine/Source/Runtime/Renderer/Private/SceneCore.cpp
--- a/Engine/Source/Runtime/Renderer/Private/SceneCore.cpp
+++ b/Engine/Source/Runtime/Renderer/Private/SceneCore.cpp
@@ -341,29 +341,18 @@ void FStaticMesh::AddToDrawLists(FRHICommandListImmediate& RHICmdList, FScene* S
 
 	if (Scene->GetShadingPath() == EShadingPath::Deferred)
 	{
-		if (bUseAsOccluder)
-		{
-			// Render non-masked materials in the depth only pass
-			extern TAutoConsoleVariable<int32> CVarEarlyZPass;
-			int32 EarlyZPass = CVarEarlyZPass.GetValueOnRenderThread();
-
-			extern int32 GEarlyZPassMovable;
-
-			EDepthDrawingMode EarlyZPassMode = (EDepthDrawingMode)EarlyZPass;
-			bool bEarlyZPassMovable = GEarlyZPassMovable != 0;
+		extern void GetEarlyZPassMode(ERHIFeatureLevel::Type FeatureLevel, EDepthDrawingMode& EarlyZPassMode, bool& bEarlyZPassMovable);
 
-			extern bool ShouldForceFullDepthPass(ERHIFeatureLevel::Type FeatureLevel);
-			if (ShouldForceFullDepthPass(Scene->GetFeatureLevel()))
-			{
-				// DBuffer decals force a full prepass
-				EarlyZPassMode = DDM_AllOccluders;
-				bEarlyZPassMovable = true;
-			}
+		EDepthDrawingMode EarlyZPassMode;
+		bool bEarlyZPassMovable;
+		GetEarlyZPassMode(Scene->GetFeatureLevel(), EarlyZPassMode, bEarlyZPassMovable);
 
+		if (bUseAsOccluder || EarlyZPassMode == DDM_AllOpaque)
+		{
 			// WARNING : If you change this condition, also change the logic in FStaticMeshSceneProxy::DrawStaticElements.
 			// Warning: also mirrored in FDeferredShadingSceneRenderer::FDeferredShadingSceneRenderer
-			if (PrimitiveSceneInfo->Proxy->ShouldUseAsOccluder() 
-				&& (!IsMasked(FeatureLevel) || EarlyZPassMode == DDM_AllOccluders)
+			if ((PrimitiveSceneInfo->Proxy->ShouldUseAsOccluder() || EarlyZPassMode == DDM_AllOpaque)
+				&& (!IsMasked(FeatureLevel) || EarlyZPassMode >= DDM_AllOccluders)
 				&& (!PrimitiveSceneInfo->Proxy->IsMovable() || bEarlyZPassMovable))
 			{
 				FDepthDrawingPolicyFactory::AddStaticMesh(Scene,this);
diff --git a/Engine/Source/Runtime/Renderer/Private/SceneRendering.cpp b/Engine/Source/Runtime/Renderer/Private/SceneRendering.cpp
--- a/Engine/Source/Runtime/Renderer/Private/SceneRendering.cpp
+++ b/Engine/Source/Runtime/Renderer/Private/SceneRendering.cpp
@@ -31,6 +31,8 @@
 #include "Matinee/MatineeActor.h"
 #include "ComponentRecreateRenderStateContext.h"
 #include "PostProcess/PostProcessSubsurface.h"
+#include "HdrCustomResolveShaders.h"
+#include "WideCustomResolveShaders.h"
 
 /*-----------------------------------------------------------------------------
 	Globals
@@ -203,6 +205,17 @@ static TAutoConsoleVariable<int32> CVarRHICmdMinDrawsPerParallelCmdList(
 	64,
 	TEXT("The minimum number of draws per cmdlist. If the total number of draws is less than this, then no parallel work will be done at all. This can't always be honored or done correctly. More effective with RHICmdBalanceParallelLists."));
 
+static TAutoConsoleVariable<int32> CVarWideCustomResolve(
+	TEXT("r.WideCustomResolve"),
+	0,
+	TEXT("Use a wide custom resolve filter when MSAA is enabled")
+	TEXT("0: Disabled [hardware box filter]")
+	TEXT("1: Wide (r=1.25, 12 samples)")
+	TEXT("2: Wider (r=1.4, 16 samples)")
+	TEXT("3: Widest (r=1.5, 20 samples)"),
+	ECVF_RenderThreadSafe | ECVF_Scalability
+	);
+
 static FParallelCommandListSet* GOutstandingParallelCommandListSet = nullptr;
 
 FParallelCommandListSet::FParallelCommandListSet(TStatId InExecuteStat, const FViewInfo& InView, FRHICommandListImmediate& InParentCmdList, bool bInParallelExecute, bool bInCreateSceneContext)
@@ -2145,3 +2158,89 @@ TSharedPtr<ISceneViewExtension, ESPMode::ThreadSafe> GetRendererViewExtension()
 }
 
 #endif // !(UE_BUILD_SHIPPING || UE_BUILD_TEST)
+
+/**
+* Saves a previously rendered scene color target
+*/
+void FSceneRenderer::ResolveSceneColor(FRHICommandList& RHICmdList)
+{
+	SCOPED_DRAW_EVENT(RHICmdList, ResolveSceneColor);
+
+	FSceneRenderTargets& SceneContext = FSceneRenderTargets::Get(RHICmdList);
+	auto& CurrentSceneColor = SceneContext.GetSceneColor();
+	uint32 CurrentNumSamples = CurrentSceneColor->GetDesc().NumSamples;
+
+	const EShaderPlatform CurrentShaderPlatform = GShaderPlatformForFeatureLevel[SceneContext.GetCurrentFeatureLevel()];
+	if (CurrentNumSamples <= 1 || !RHISupportsSeparateMSAAAndResolveTextures(CurrentShaderPlatform))
+	{
+		RHICmdList.CopyToResolveTarget(SceneContext.GetSceneColorSurface(), SceneContext.GetSceneColorTexture(), true, FResolveRect(0, 0, ViewFamily.FamilySizeX, ViewFamily.FamilySizeY));
+	}
+	else 
+	{
+		// Custom shader based color resolve for HDR color to emulate mobile.
+		SetRenderTarget(RHICmdList, SceneContext.GetSceneColorTexture(), FTextureRHIParamRef());
+		
+		for (int32 ViewIndex = 0; ViewIndex < Views.Num(); ViewIndex++)
+		{
+			const FViewInfo& View = Views[ViewIndex];
+
+			RHICmdList.SetBlendState(TStaticBlendState<>::GetRHI());
+			RHICmdList.SetRasterizerState(TStaticRasterizerState<>::GetRHI());
+			RHICmdList.SetDepthStencilState(TStaticDepthStencilState<false, CF_Always>::GetRHI());
+			RHICmdList.SetStreamSource(0, NULL, 0, 0);
+
+			// Resolve views individually
+			// In the case of adaptive resolution, the view family will be much larger than the views individually
+			RHICmdList.SetScissorRect(true, View.ViewRect.Min.X, View.ViewRect.Min.Y, View.ViewRect.Max.X, View.ViewRect.Max.Y);
+
+			int32 ResolveWidth = CVarWideCustomResolve.GetValueOnRenderThread();
+
+			if (CurrentNumSamples <= 1)
+			{
+				ResolveWidth = 0;
+			}
+
+			if (ResolveWidth != 0)
+			{
+				ResolveFilterWide(RHICmdList, SceneContext.GetCurrentFeatureLevel(), CurrentSceneColor->GetRenderTargetItem().TargetableTexture, FIntPoint(0, 0), CurrentNumSamples, ResolveWidth);
+			}
+			else
+			{
+				auto ShaderMap = GetGlobalShaderMap(SceneContext.GetCurrentFeatureLevel());
+				TShaderMapRef<FHdrCustomResolveVS> VertexShader(ShaderMap);
+
+				if (CurrentNumSamples == 2)
+				{
+					TShaderMapRef<FHdrCustomResolve2xPS> PixelShader(ShaderMap);
+					static FGlobalBoundShaderState BoundShaderState;
+					SetGlobalBoundShaderState(RHICmdList, SceneContext.GetCurrentFeatureLevel(), BoundShaderState, GetVertexDeclarationFVector4(), *VertexShader, *PixelShader);
+					PixelShader->SetParameters(RHICmdList, CurrentSceneColor->GetRenderTargetItem().TargetableTexture);
+					RHICmdList.DrawPrimitive(PT_TriangleList, 0, 1, 1);
+				}
+				else if (CurrentNumSamples == 4)
+				{
+					TShaderMapRef<FHdrCustomResolve4xPS> PixelShader(ShaderMap);
+					static FGlobalBoundShaderState BoundShaderState;
+					SetGlobalBoundShaderState(RHICmdList, SceneContext.GetCurrentFeatureLevel(), BoundShaderState, GetVertexDeclarationFVector4(), *VertexShader, *PixelShader);
+					PixelShader->SetParameters(RHICmdList, CurrentSceneColor->GetRenderTargetItem().TargetableTexture);
+					RHICmdList.DrawPrimitive(PT_TriangleList, 0, 1, 1);
+				}
+				else if (CurrentNumSamples == 8)
+				{
+					TShaderMapRef<FHdrCustomResolve8xPS> PixelShader(ShaderMap);
+					static FGlobalBoundShaderState BoundShaderState;
+					SetGlobalBoundShaderState(RHICmdList, SceneContext.GetCurrentFeatureLevel(), BoundShaderState, GetVertexDeclarationFVector4(), *VertexShader, *PixelShader);
+					PixelShader->SetParameters(RHICmdList, CurrentSceneColor->GetRenderTargetItem().TargetableTexture);
+					RHICmdList.DrawPrimitive(PT_TriangleList, 0, 1, 1);
+				}
+				else
+				{
+					// Everything other than 2,4,8 samples is not implemented.
+					check(0);
+				}
+			}
+		}
+
+		RHICmdList.SetScissorRect(false, 0, 0, 0, 0);
+	}
+}
\ No newline at end of file
diff --git a/Engine/Source/Runtime/Renderer/Private/SceneRendering.h b/Engine/Source/Runtime/Renderer/Private/SceneRendering.h
--- a/Engine/Source/Runtime/Renderer/Private/SceneRendering.h
+++ b/Engine/Source/Runtime/Renderer/Private/SceneRendering.h
@@ -1397,6 +1397,8 @@ protected:
 	void ClearPrimitiveSingleFramePrecomputedLightingBuffers();
 
 	void RenderPlanarReflection(class FPlanarReflectionSceneProxy* ReflectionSceneProxy);
+
+	void ResolveSceneColor(FRHICommandList& RHICmdList);
 };
 
 /**
diff --git a/Engine/Source/Runtime/Renderer/Private/SceneVisibility.cpp b/Engine/Source/Runtime/Renderer/Private/SceneVisibility.cpp
--- a/Engine/Source/Runtime/Renderer/Private/SceneVisibility.cpp
+++ b/Engine/Source/Runtime/Renderer/Private/SceneVisibility.cpp
@@ -1357,7 +1357,7 @@ struct FMarkRelevantStaticMeshesForViewData
 	float InvLODScale;
 	float MinScreenRadiusForCSMDepthSquared;
 	float MinScreenRadiusForDepthPrepassSquared;
-	bool bForceEarlyZPass;
+	bool bFullEarlyZPass;
 
 	FMarkRelevantStaticMeshesForViewData(FViewInfo& View)
 	{
@@ -1368,14 +1368,14 @@ struct FMarkRelevantStaticMeshesForViewData
 		// outside of the loop to be more efficient
 		ForcedLODLevel = (View.Family->EngineShowFlags.LOD) ? GetCVarForceLOD() : 0;
 
-		LODScale = CVarStaticMeshLODDistanceScale.GetValueOnRenderThread();
+		LODScale = CVarStaticMeshLODDistanceScale.GetValueOnRenderThread() * View.LODDistanceFactor;
 		InvLODScale = 1.0f / LODScale;
 
 		MinScreenRadiusForCSMDepthSquared = GMinScreenRadiusForCSMDepth * GMinScreenRadiusForCSMDepth;
 		MinScreenRadiusForDepthPrepassSquared = GMinScreenRadiusForDepthPrepass * GMinScreenRadiusForDepthPrepass;
 
-		extern TAutoConsoleVariable<int32> CVarEarlyZPass;
-		bForceEarlyZPass = CVarEarlyZPass.GetValueOnRenderThread() == 2;
+		extern bool ShouldForceFullDepthPass(ERHIFeatureLevel::Type FeatureLevel);
+		bFullEarlyZPass = ShouldForceFullDepthPass(View.GetFeatureLevel());
 	}
 };
 
@@ -1624,7 +1624,7 @@ struct FRelevancePacket
 			float DistanceSquared = (Bounds.Origin - ViewData.ViewOrigin).SizeSquared();
 			const float LODFactorDistanceSquared = DistanceSquared * FMath::Square(View.LODDistanceFactor * ViewData.InvLODScale);
 			const bool bDrawShadowDepth = FMath::Square(Bounds.SphereRadius) > ViewData.MinScreenRadiusForCSMDepthSquared * LODFactorDistanceSquared;
-			const bool bDrawDepthOnly = ViewData.bForceEarlyZPass || FMath::Square(Bounds.SphereRadius) > GMinScreenRadiusForDepthPrepass * GMinScreenRadiusForDepthPrepass * LODFactorDistanceSquared;
+			const bool bDrawDepthOnly = ViewData.bFullEarlyZPass || FMath::Square(Bounds.SphereRadius) > GMinScreenRadiusForDepthPrepass * GMinScreenRadiusForDepthPrepass * LODFactorDistanceSquared;
 
 			const int32 NumStaticMeshes = PrimitiveSceneInfo->StaticMeshes.Num();
 			for(int32 MeshIndex = 0;MeshIndex < NumStaticMeshes;MeshIndex++)
diff --git a/Engine/Source/Runtime/Renderer/Private/TiledDeferredLightRendering.cpp b/Engine/Source/Runtime/Renderer/Private/TiledDeferredLightRendering.cpp
--- a/Engine/Source/Runtime/Renderer/Private/TiledDeferredLightRendering.cpp
+++ b/Engine/Source/Runtime/Renderer/Private/TiledDeferredLightRendering.cpp
@@ -359,7 +359,7 @@ void FDeferredShadingSceneRenderer::RenderTiledDeferredLighting(FRHICommandListI
 			// One some hardware we can read and write from the same UAV with a 32 bit format. We don't do that yet.
 			TRefCountPtr<IPooledRenderTarget> OutTexture;
 			{
-				SceneContext.ResolveSceneColor(RHICmdList, FResolveRect(0, 0, ViewFamily.FamilySizeX, ViewFamily.FamilySizeY));
+				ResolveSceneColor(RHICmdList);
 
 				FPooledRenderTargetDesc Desc = SceneContext.GetSceneColor()->GetDesc();
 				Desc.TargetableFlags |= TexCreate_UAV;
diff --git a/Engine/Source/Runtime/Renderer/Private/TranslucentRendering.cpp b/Engine/Source/Runtime/Renderer/Private/TranslucentRendering.cpp
--- a/Engine/Source/Runtime/Renderer/Private/TranslucentRendering.cpp
+++ b/Engine/Source/Runtime/Renderer/Private/TranslucentRendering.cpp
@@ -177,21 +177,6 @@ static void SetTranslucentState(FRHICommandList& RHICmdList, FDrawingPolicyRende
 	DrawRenderState.SetDepthStencilState(RHICmdList, TStaticDepthStencilState<false, CF_DepthNearOrEqual>::GetRHI());
 }
 
-static void FinishTranslucentRenderTarget(FRHICommandListImmediate& RHICmdList, const FViewInfo& View, ETranslucencyPass::Type TranslucencyPass)
-{
-	QUICK_SCOPE_CYCLE_COUNTER(STAT_FinishTranslucentRenderTarget);
-	
-	FSceneRenderTargets& SceneContext = FSceneRenderTargets::Get(RHICmdList);
-	if ((TranslucencyPass == ETranslucencyPass::TPT_SeparateTranslucency) && SceneContext.IsSeparateTranslucencyActive(View))
-	{
-		SceneContext.FinishRenderingSeparateTranslucency(RHICmdList);
-	}
-	else
-	{
-		SceneContext.FinishRenderingTranslucency(RHICmdList, View);
-	}
-}
-
 const FProjectedShadowInfo* FDeferredShadingSceneRenderer::PrepareTranslucentShadowMap(FRHICommandList& RHICmdList, const FViewInfo& View, FPrimitiveSceneInfo* PrimitiveSceneInfo, ETranslucencyPass::Type TranslucencyPass)
 {
 	QUICK_SCOPE_CYCLE_COUNTER(STAT_FDeferredShadingSceneRenderer_PrepareTranslucentShadowMap);
@@ -278,7 +263,7 @@ void FTranslucencyDrawingPolicyFactory::CopySceneColor(FRHICommandList& RHICmdLi
 	RHICmdList.SetDepthStencilState(TStaticDepthStencilState<false, CF_Always>::GetRHI());
 	RHICmdList.SetBlendState(TStaticBlendState<>::GetRHI());
 
-	SceneContext.ResolveSceneColor(RHICmdList);
+	RHICmdList.CopyToResolveTarget(SceneContext.GetSceneColorSurface(), SceneContext.GetSceneColorTexture(), true, FResolveRect(View.ViewRect.Min.X, View.ViewRect.Min.Y, View.ViewRect.Max.X, View.ViewRect.Max.Y));
 
 	SceneContext.BeginRenderingLightAttenuation(RHICmdList);
 	RHICmdList.SetViewport(View.ViewRect.Min.X, View.ViewRect.Min.Y, 0.0f, View.ViewRect.Max.X, View.ViewRect.Max.Y, 1.0f);
diff --git a/Engine/Source/Runtime/Renderer/Public/DrawingPolicy.h b/Engine/Source/Runtime/Renderer/Public/DrawingPolicy.h
--- a/Engine/Source/Runtime/Renderer/Public/DrawingPolicy.h
+++ b/Engine/Source/Runtime/Renderer/Public/DrawingPolicy.h
@@ -44,6 +44,7 @@ struct FDrawingPolicyRenderState
 #endif
 		  BlendState(nullptr)
 		, DepthStencilState(nullptr)
+		, DepthStencilAccess(FExclusiveDepthStencil::DepthRead_StencilRead)
 		, ViewUniformBuffer(SceneView.ViewUniformBuffer)
 		, StencilRef(0)
 		, ViewOverrideFlags(EDrawingPolicyOverrideFlags::None)
@@ -64,6 +65,7 @@ struct FDrawingPolicyRenderState
 #endif
 		  BlendState(DrawRenderState.BlendState)
 		, DepthStencilState(DrawRenderState.DepthStencilState)
+		, DepthStencilAccess(DrawRenderState.DepthStencilAccess)
 		, ViewUniformBuffer(DrawRenderState.ViewUniformBuffer)
 		, StencilRef(DrawRenderState.StencilRef)
 		, ViewOverrideFlags(DrawRenderState.ViewOverrideFlags)
@@ -133,6 +135,16 @@ public:
 		return DepthStencilState;
 	}
 
+	FORCEINLINE_DEBUGGABLE void SetDepthStencilAccess(FExclusiveDepthStencil::Type InDepthStencilAccess)
+	{
+		DepthStencilAccess = InDepthStencilAccess;
+	}
+
+	FORCEINLINE_DEBUGGABLE FExclusiveDepthStencil::Type GetDepthStencilAccess() const
+	{
+		return DepthStencilAccess;
+	}
+
 	FORCEINLINE_DEBUGGABLE void SetViewUniformBuffer(const TUniformBufferRef<FViewUniformShaderParameters>& InViewUniformBuffer)
 	{
 		ViewUniformBuffer = InViewUniformBuffer;
@@ -187,6 +199,7 @@ private:
 
 	FBlendStateRHIParamRef			BlendState;
 	FDepthStencilStateRHIParamRef	DepthStencilState;
+	FExclusiveDepthStencil::Type	DepthStencilAccess;
 
 	TUniformBufferRef<FViewUniformShaderParameters>	ViewUniformBuffer;
 	uint32							StencilRef;
diff --git a/Engine/Source/Runtime/SlateRHIRenderer/Private/SlateRHIRenderer.cpp b/Engine/Source/Runtime/SlateRHIRenderer/Private/SlateRHIRenderer.cpp
--- a/Engine/Source/Runtime/SlateRHIRenderer/Private/SlateRHIRenderer.cpp
+++ b/Engine/Source/Runtime/SlateRHIRenderer/Private/SlateRHIRenderer.cpp
@@ -310,7 +310,7 @@ void FSlateRHIRenderer::CreateViewport( const TSharedRef<SWindow> Window )
 		// SDR format holds the requested format in non HDR mode
 		NewInfo->SDRPixelFormat = NewInfo->PixelFormat;
 
-		if (CVarHDROutputEnabled->GetValueOnGameThread() != 0)
+		if (CVarHDROutputEnabled && CVarHDROutputEnabled->GetValueOnGameThread() != 0)
 		{
 			NewInfo->PixelFormat = GRHIHDRDisplayOutputFormat;
 		}
@@ -337,9 +337,9 @@ void FSlateRHIRenderer::ConditionalResizeViewport( FViewportInfo* ViewInfo, uint
 	static const auto CVarHDRColorGamut = IConsoleManager::Get().FindTConsoleVariableDataInt(TEXT("r.HDR.Display.ColorGamut"));
 	static const auto CVarHDROutputDevice = IConsoleManager::Get().FindTConsoleVariableDataInt(TEXT("r.HDR.Display.OutputDevice"));
 
-	bool bHDREnabled = GRHISupportsHDROutput && CVarHDROutputEnabled->GetValueOnAnyThread() != 0;
-	int32 HDRColorGamut = CVarHDRColorGamut->GetValueOnAnyThread();
-	int32 HDROutputDevice = CVarHDROutputDevice->GetValueOnAnyThread();
+	bool bHDREnabled = GRHISupportsHDROutput && CVarHDROutputEnabled && CVarHDROutputEnabled->GetValueOnAnyThread() != 0;
+	int32 HDRColorGamut = CVarHDRColorGamut ? CVarHDRColorGamut->GetValueOnAnyThread() : 0;
+	int32 HDROutputDevice = CVarHDROutputDevice ? CVarHDROutputDevice->GetValueOnAnyThread() : 0;
 
 	bool bHDRStale = ViewInfo && (
 		((ViewInfo->PixelFormat == GRHIHDRDisplayOutputFormat) != bHDREnabled)	// HDR toggled
@@ -633,7 +633,9 @@ void FSlateRHIRenderer::DrawWindow_RenderThread(FRHICommandListImmediate& RHICmd
 	static const auto CVarHDROutputEnabled = IConsoleManager::Get().FindTConsoleVariableDataInt(TEXT("r.HDR.EnableHDROutput"));
 
 	const bool bSupportsUIComposition = GRHISupportsHDROutput && GSupportsVolumeTextureRendering && SupportsUICompositionRendering(GetFeatureLevelShaderPlatform(GMaxRHIFeatureLevel));
-	const bool bCompositeUI = bSupportsUIComposition && CVarCompositeMode->GetValueOnRenderThread() != 0 && CVarHDROutputEnabled->GetValueOnRenderThread() != 0;
+	const bool bCompositeUI = bSupportsUIComposition
+		&& CVarCompositeMode && CVarCompositeMode->GetValueOnRenderThread() != 0
+		&& CVarHDROutputEnabled && CVarHDROutputEnabled->GetValueOnRenderThread() != 0;
 
 	const int32 CompositionLUTSize = 32;
 
@@ -641,8 +643,8 @@ void FSlateRHIRenderer::DrawWindow_RenderThread(FRHICommandListImmediate& RHICmd
 	static const auto CVarHDROutputDevice = IConsoleManager::Get().FindTConsoleVariableDataInt(TEXT("r.HDR.Display.OutputDevice"));
 	static const auto CVarHDROutputGamut = IConsoleManager::Get().FindTConsoleVariableDataInt(TEXT("r.HDR.Display.ColorGamut"));
 
-	const int32 HDROutputDevice = CVarHDROutputDevice->GetValueOnRenderThread();
-	const int32 HDROutputGamut = CVarHDROutputGamut->GetValueOnRenderThread();
+	const int32 HDROutputDevice = CVarHDROutputDevice ? CVarHDROutputDevice->GetValueOnRenderThread() : 0;
+	const int32 HDROutputGamut = CVarHDROutputGamut ? CVarHDROutputGamut->GetValueOnRenderThread() : 0;
 
 	bool bLUTStale = ViewportInfo.ColorSpaceLUTOutputDevice != HDROutputDevice || ViewportInfo.ColorSpaceLUTOutputGamut != HDROutputGamut;
 
diff --git a/Engine/Source/Runtime/VulkanRHI/Private/VulkanRenderTarget.cpp b/Engine/Source/Runtime/VulkanRHI/Private/VulkanRenderTarget.cpp
--- a/Engine/Source/Runtime/VulkanRHI/Private/VulkanRenderTarget.cpp
+++ b/Engine/Source/Runtime/VulkanRHI/Private/VulkanRenderTarget.cpp
@@ -461,6 +461,8 @@ void FVulkanDynamicRHI::RHIReadSurfaceData(FTextureRHIParamRef TextureRHI, FIntR
 	}
 
 	VkBufferMemoryBarrier Barrier;
+	ensure(StagingBuffer->GetSize() == Size);
+		//#todo-rco: Change offset if reusing a buffer suballocation
 	VulkanRHI::SetupAndZeroBufferBarrier(Barrier, VK_ACCESS_TRANSFER_WRITE_BIT, VK_ACCESS_MEMORY_READ_BIT, StagingBuffer->GetHandle(), 0/*StagingBuffer->GetOffset()*/, Size);
 	VulkanRHI::vkCmdPipelineBarrier(CmdBuffer->GetHandle(), VK_PIPELINE_STAGE_TRANSFER_BIT, VK_PIPELINE_STAGE_HOST_BIT, 0, 0, nullptr, 1, &Barrier, 0, nullptr);
 
@@ -546,8 +548,9 @@ void FVulkanDynamicRHI::RHIReadSurfaceFloatData(FTextureRHIParamRef TextureRHI,
 		}
 
 		VkBufferMemoryBarrier Barrier;
+		ensure(StagingBuffer->GetSize() == Size);
 		//#todo-rco: Change offset if reusing a buffer suballocation
-		VulkanRHI::SetupAndZeroBufferBarrier(Barrier, VK_ACCESS_TRANSFER_WRITE_BIT, VK_ACCESS_HOST_READ_BIT, StagingBuffer->GetHandle(), 0/*StagingBuffer->GetOffset()*/, StagingBuffer->GetDeviceMemoryAllocationSize());
+		VulkanRHI::SetupAndZeroBufferBarrier(Barrier, VK_ACCESS_TRANSFER_WRITE_BIT, VK_ACCESS_HOST_READ_BIT, StagingBuffer->GetHandle(), 0/*StagingBuffer->GetOffset()*/, StagingBuffer->GetSize());
 		VulkanRHI::vkCmdPipelineBarrier(InCmdBuffer->GetHandle(), VK_PIPELINE_STAGE_TRANSFER_BIT, VK_PIPELINE_STAGE_HOST_BIT, 0, 0, nullptr, 1, &Barrier, 0, nullptr);
 
 		// Force upload
diff --git a/Engine/Source/Runtime/VulkanRHI/Public/VulkanMemory.h b/Engine/Source/Runtime/VulkanRHI/Public/VulkanMemory.h
--- a/Engine/Source/Runtime/VulkanRHI/Public/VulkanMemory.h
+++ b/Engine/Source/Runtime/VulkanRHI/Public/VulkanMemory.h
@@ -843,9 +843,9 @@ namespace VulkanRHI
 			return ResourceAllocation->GetOffset();
 		}
 
-		inline uint32 GetDeviceMemoryAllocationSize() const
+		inline uint32 GetSize() const
 		{
-			return ResourceAllocation->GetAllocationSize();
+			return ResourceAllocation->GetSize();
 		}
 
 		inline VkDeviceMemory GetDeviceMemoryHandle() const
diff --git a/Engine/Source/Runtime/Windows/D3D11RHI/Private/Windows/WindowsD3D11Device.cpp b/Engine/Source/Runtime/Windows/D3D11RHI/Private/Windows/WindowsD3D11Device.cpp
--- a/Engine/Source/Runtime/Windows/D3D11RHI/Private/Windows/WindowsD3D11Device.cpp
+++ b/Engine/Source/Runtime/Windows/D3D11RHI/Private/Windows/WindowsD3D11Device.cpp
@@ -310,7 +310,7 @@ void FD3D11DynamicRHI::EnableHDR()
 	static const auto CVarHDRColorGamut = IConsoleManager::Get().FindTConsoleVariableDataInt(TEXT("r.HDR.Display.ColorGamut"));
 	static const auto CVarHDROutputDevice = IConsoleManager::Get().FindTConsoleVariableDataInt(TEXT("r.HDR.Display.OutputDevice"));
 
-	if (GRHISupportsHDROutput && CVarHDROutputEnabled->GetValueOnAnyThread() != 0)
+	if (GRHISupportsHDROutput && CVarHDROutputEnabled && CVarHDROutputEnabled->GetValueOnAnyThread() != 0)
 	{
 		const int32 OutputDevice = CVarHDROutputDevice->GetValueOnAnyThread();
 
