Fix compile error when #STATS not defined

diff --git a/Engine/Plugins/Performance/PerformanceMonitor/Source/PerformanceMonitor/Private/PerformanceMonitor.h b/Engine/Plugins/Performance/PerformanceMonitor/Source/PerformanceMonitor/Private/PerformanceMonitor.h
--- a/Engine/Plugins/Performance/PerformanceMonitor/Source/PerformanceMonitor/Private/PerformanceMonitor.h
+++ b/Engine/Plugins/Performance/PerformanceMonitor/Source/PerformanceMonitor/Private/PerformanceMonitor.h
@@ -83,9 +83,7 @@ public:
 
 
 	bool IsRecordingPerfTimers() { return bRecording; }
-#if STATS
 	void StartRecordingPerfTimers(FString FileName, TArray<FString> StatsToRecord);
-#endif
 	void SetRecordInterval(float NewInterval);
 
 	void RecordFrame();
